Add initial plugin AnimNext warping
- Add Steering additive Trait

#tests PIE Locomotion

diff --git a/Engine/Plugins/Experimental/AnimNextWarping/AnimNextWarping.uplugin b/Engine/Plugins/Experimental/AnimNextWarping/AnimNextWarping.uplugin
new file mode 100644
--- /dev/null
+++ b/Engine/Plugins/Experimental/AnimNextWarping/AnimNextWarping.uplugin
@@ -0,0 +1,41 @@
+{
+	"FileVersion" : 3,
+	"Version" : 1,
+	"VersionName" : "0.1",
+	"FriendlyName" : "Anim Next Warping",
+	"Description": "Framework for animation and pose warping for Anim Next.",
+	"Category" : "Animation",
+	"CreatedBy" : "Epic Games, Inc.",
+	"CreatedByURL" : "https://epicgames.com",
+	"DocsURL" : "",
+	"MarketplaceURL" : "",
+	"SupportURL" : "",
+	"EnabledByDefault" : false,
+	"CanContainContent" : false,
+	"IsBetaVersion": false,
+	"IsExperimentalVersion": true,
+	"Installed" : false,
+	"Plugins":
+	[
+		{
+			"Name": "AnimNext",
+			"Enabled": true
+		},
+		{
+			"Name": "AnimNextAnimGraph",
+			"Enabled": true
+		},
+		{
+			"Name": "ControlRig",
+			"Enabled": true
+		}
+	],
+	"Modules" :
+	[
+		{
+			"Name" : "AnimNextWarping",
+			"Type" : "Runtime",
+			"LoadingPhase" : "Default"
+		}
+	]
+}
diff --git a/Engine/Plugins/Experimental/AnimNextWarping/Source/AnimNextWarping/AnimNextWarping.Build.cs b/Engine/Plugins/Experimental/AnimNextWarping/Source/AnimNextWarping/AnimNextWarping.Build.cs
new file mode 100644
--- /dev/null
+++ b/Engine/Plugins/Experimental/AnimNextWarping/Source/AnimNextWarping/AnimNextWarping.Build.cs
@@ -0,0 +1,23 @@
+// Copyright Epic Games, Inc. All Rights Reserved.
+
+namespace UnrealBuildTool.Rules
+{
+	public class AnimNextWarping : ModuleRules
+	{
+		public AnimNextWarping(ReadOnlyTargetRules Target) : base(Target)
+		{
+			PrivateDependencyModuleNames.AddRange(
+				new string[]
+				{
+					"Core",
+					"CoreUObject",
+					"RigVM",
+					"ControlRig",
+					"Engine",
+					"AnimNext",
+					"AnimNextAnimGraph",
+				}
+			);
+		}
+	}
+}
\ No newline at end of file
diff --git a/Engine/Plugins/Experimental/AnimNextWarping/Source/AnimNextWarping/Private/AnimNextWarpingLog.h b/Engine/Plugins/Experimental/AnimNextWarping/Source/AnimNextWarping/Private/AnimNextWarpingLog.h
new file mode 100644
--- /dev/null
+++ b/Engine/Plugins/Experimental/AnimNextWarping/Source/AnimNextWarping/Private/AnimNextWarpingLog.h
@@ -0,0 +1,7 @@
+// Copyright Epic Games, Inc. All Rights Reserved.
+
+#pragma once
+
+#include "CoreMinimal.h"
+
+DECLARE_LOG_CATEGORY_EXTERN(LogAnimNextWarping, Log, All);
diff --git a/Engine/Plugins/Experimental/AnimNextWarping/Source/AnimNextWarping/Private/Module.cpp b/Engine/Plugins/Experimental/AnimNextWarping/Source/AnimNextWarping/Private/Module.cpp
new file mode 100644
--- /dev/null
+++ b/Engine/Plugins/Experimental/AnimNextWarping/Source/AnimNextWarping/Private/Module.cpp
@@ -0,0 +1,40 @@
+// Copyright Epic Games, Inc. All Rights Reserved.
+
+#include "CoreMinimal.h"
+#include "Module/AnimNextModule.h"
+#include "Modules/ModuleInterface.h"
+#include "Modules/ModuleManager.h"
+#include "RigVMCore/RigVMRegistry.h"
+#include "Animation/AnimSequence.h"
+
+#include "AnimNextWarpingLog.h"
+#include "EvaluationVM/EvaluationVM.h"
+
+DEFINE_LOG_CATEGORY(LogAnimNextWarping);
+
+namespace UE::AnimNext::Warping
+{
+
+class FModule : public IModuleInterface
+{
+public:
+	virtual void StartupModule() override
+	{
+		// @TODO: Register any pin types needed for warping here for RigVM, currently we have none
+		//static TPair<UClass*, FRigVMRegistry::ERegisterObjectOperation> const AllowedObjectTypes[] =
+		//{
+		//	// { UMyRigVMParamClass::StaticClass(), FRigVMRegistry::ERegisterObjectOperation::Class },
+		//};
+
+		//FRigVMRegistry::Get().RegisterObjectTypes(AllowedObjectTypes);
+	}
+
+	virtual void ShutdownModule() override
+	{
+	}
+
+};
+
+}
+
+IMPLEMENT_MODULE(UE::AnimNext::Warping::FModule, AnimNextWarping)
diff --git a/Engine/Plugins/Experimental/AnimNextWarping/Source/AnimNextWarping/Private/SteeringTrait.cpp b/Engine/Plugins/Experimental/AnimNextWarping/Source/AnimNextWarping/Private/SteeringTrait.cpp
new file mode 100644
--- /dev/null
+++ b/Engine/Plugins/Experimental/AnimNextWarping/Source/AnimNextWarping/Private/SteeringTrait.cpp
@@ -0,0 +1,225 @@
+// Copyright Epic Games, Inc. All Rights Reserved.
+
+#include "SteeringTrait.h"
+
+#include "AnimNextAnimGraphSettings.h"
+#include "AnimNextWarpingLog.h"
+#include "Animation/AnimRootMotionProvider.h"
+#include "Animation/AnimSequence.h"
+#include "DataInterface/DataInterfaceStructAdapter.h"
+#include "EvaluationVM/EvaluationVM.h"
+#include "EvaluationVM/KeyframeState.h"
+#include "GraphInterfaces/AnimNextNativeDataInterface_AnimSequencePlayer.h"
+#include "TraitInterfaces/IBlendStack.h"
+#include "TraitInterfaces/IGraphFactory.h"
+#include "TraitInterfaces/ITimeline.h"
+
+////////////////////////////////////////////////////////////////////////////////////////////////////////////////////////
+// FSteeringTrait
+
+namespace UE::AnimNext
+{
+	AUTO_REGISTER_ANIM_TRAIT(FSteeringTrait)
+
+	#define TRAIT_INTERFACE_ENUMERATOR(GeneratorMacro) \
+		GeneratorMacro(IUpdate) \
+		GeneratorMacro(IEvaluate) \
+
+	// Trait implementation boilerplate
+	GENERATE_ANIM_TRAIT_IMPLEMENTATION(FSteeringTrait, TRAIT_INTERFACE_ENUMERATOR, NULL_ANIM_TRAIT_INTERFACE_ENUMERATOR, NULL_ANIM_TRAIT_EVENT_ENUMERATOR)
+	#undef TRAIT_INTERFACE_ENUMERATOR
+
+	void FSteeringTrait::PreUpdate(FUpdateTraversalContext& Context, const TTraitBinding<IUpdate>& Binding, const FTraitUpdateState& TraitState) const
+	{
+		FInstanceData* InstanceData = Binding.GetInstanceData<FInstanceData>();
+		check(InstanceData);
+
+		InstanceData->DeltaTime = TraitState.GetDeltaTime();
+
+		IUpdate::PreUpdate(Context, Binding, TraitState);
+	}
+
+	void FSteeringTrait::PostEvaluate(FEvaluateTraversalContext& Context, const TTraitBinding<IEvaluate>& Binding) const
+	{
+		IEvaluate::PostEvaluate(Context, Binding);
+
+		const FSharedData* SharedData = Binding.GetSharedData<FSharedData>();
+		check(SharedData);
+
+		FInstanceData* InstanceData = Binding.GetInstanceData<FInstanceData>();
+		check(InstanceData);
+
+		// Get Current Anim Asset
+		{
+			TTraitBinding<IBlendStack> BlendStackTrait;
+			if (!Binding.GetStackInterface(BlendStackTrait))
+			{
+				UE_LOG(LogAnimNextWarping, Error, TEXT("FSteeringTrait::PostEvaluate, missing IBlendStack"));
+				return;
+			}
+
+			// @TODO: See if we want this as a trait interface (Some trait to get anim asset). 
+			auto GetCurrentAnimAsset = [&]()
+			{
+				TObjectPtr<const UAnimSequence> Result = nullptr;
+
+				IBlendStack::FGraphRequestPtr ActiveGraphRequest;
+				BlendStackTrait.GetActiveGraph(Context, ActiveGraphRequest);
+
+				if (ActiveGraphRequest)
+				{
+					for (const FInstancedStruct& Payload : ActiveGraphRequest->GraphPayload.GetNativePayloads())
+					{
+						if (const FAnimNextNativeDataInterface_AnimSequencePlayer* PlayAnimPayload = Payload.GetPtr<FAnimNextNativeDataInterface_AnimSequencePlayer>())
+						{
+							Result = PlayAnimPayload->AnimSequence;
+							return Result;
+						}
+					}
+				}
+
+				return Result;
+			};
+
+			InstanceData->CurrentAnimAsset = GetCurrentAnimAsset();
+		}
+
+		// Get Current Anim Asset Time
+		{
+			TTraitBinding<ITimeline> TimelineTrait;
+			if (!Binding.GetStackInterface(TimelineTrait))
+			{
+				UE_LOG(LogAnimNextWarping, Error, TEXT("FSteeringTrait::PostEvaluate, missing ITimeline"));
+				return;
+			}
+
+			InstanceData->CurrentAnimAssetTime = TimelineTrait.GetProgress(Context).GetPosition();
+		}
+
+		// Update target orientation / root motion
+		{
+			InstanceData->TargetOrientation = SharedData->GetTargetOrientation(Binding);
+			InstanceData->RootBoneTransform = SharedData->GetRootBoneTransform(Binding);
+		}
+
+		Context.AppendTask(FAnimNextSteeringTask::Make(InstanceData, SharedData));
+	}
+} // namespace UE::AnimNext
+
+
+////////////////////////////////////////////////////////////////////////////////////////////////////////////////////////
+// FAnimNextSteeringTask
+
+FAnimNextSteeringTask FAnimNextSteeringTask::Make(UE::AnimNext::FSteeringTrait::FInstanceData* InstanceData, const UE::AnimNext::FSteeringTrait::FSharedData* SharedData)
+{
+	FAnimNextSteeringTask Task;
+	Task.InstanceData = InstanceData;
+	Task.SharedData = SharedData;
+	return Task;
+}
+
+void FAnimNextSteeringTask::Execute(UE::AnimNext::FEvaluationVM& VM) const
+{
+	using namespace UE::AnimNext;
+
+	// Execute Steering
+	if (InstanceData->DeltaTime > 0)
+	{
+		if (SharedData->bEnableTargetSmoothing)
+		{
+			if (InstanceData->bResetFilter)
+			{
+				InstanceData->FilteredTarget = InstanceData->TargetOrientation;
+				InstanceData->TargetSmoothingState.Reset();
+				InstanceData->bResetFilter = false;
+			}
+
+			InstanceData->FilteredTarget = UKismetMathLibrary::QuaternionSpringInterp(InstanceData->FilteredTarget, InstanceData->TargetOrientation, InstanceData->TargetSmoothingState,
+				SharedData->SmoothTargetStiffness, SharedData->SmoothTargetDamping, InstanceData->DeltaTime, 1, 0, true);
+
+			InstanceData->TargetOrientation = InstanceData->FilteredTarget;
+		}
+
+		if (SharedData->Alpha > 0.0f)
+		{
+			const UE::Anim::IAnimRootMotionProvider* RootMotionProvider = UE::Anim::IAnimRootMotionProvider::Get();
+			if (!RootMotionProvider)
+			{
+				UE_LOG(LogAnimNextWarping, Error, TEXT("FAnimNextSteeringTask::Execute, missing RootMotionProvider"));
+				return;
+			}
+
+			if (const TUniquePtr<FKeyframeState>* Keyframe = VM.PeekValue<TUniquePtr<FKeyframeState>>(KEYFRAME_STACK_NAME, 0))
+			{
+				FTransform ThisFrameRootMotionTransform = FTransform::Identity;
+				if (RootMotionProvider->ExtractRootMotion(Keyframe->Get()->Attributes, ThisFrameRootMotionTransform))
+				{
+					if (ThisFrameRootMotionTransform.GetTranslation().Length() / InstanceData->DeltaTime > SharedData->DisableSteeringBelowSpeed)
+					{
+						FQuat RootBoneRotation = InstanceData->RootBoneTransform.GetRotation();
+
+						//UE_VLOG_ARROW(Output.AnimInstanceProxy->GetAnimInstanceObject(), "Steering", Display,
+						//	RootBoneTransform.GetLocation(),
+						//	RootBoneTransform.GetLocation() + RootBoneRotation.GetRightVector() * 90,
+						//	FColor::Green, TEXT(""));
+
+						//UE_VLOG_ARROW(Output.AnimInstanceProxy->GetAnimInstanceObject(), "Steering", Display,
+						//	RootBoneTransform.GetLocation(),
+						//	RootBoneTransform.GetLocation() + InstanceData->TargetOrientation.GetRightVector() * 100,
+						//	FColor::Blue, TEXT(""));
+
+						FQuat DeltaToTargetOrientation = RootBoneRotation.Inverse() * InstanceData->TargetOrientation;
+
+						if (SharedData->AnimatedTargetTime > 0)
+						{
+							if (const UAnimSequenceBase* AnimSequence = InstanceData->CurrentAnimAsset)
+							{
+								FTransform PredictedRootMotionTransform = AnimSequence->ExtractRootMotion(InstanceData->CurrentAnimAssetTime, SharedData->AnimatedTargetTime, true);
+								FQuat PredictedRootMotionQuat = PredictedRootMotionTransform.GetRotation();
+								float PredictedRootMotionYaw = FRotator(PredictedRootMotionQuat).Yaw;
+
+								if (fabs(PredictedRootMotionYaw) > SharedData->RootMotionThreshold)
+								{
+									//UE_VLOG_ARROW(Output.AnimInstanceProxy->GetAnimInstanceObject(), "Steering", Display,
+									//	RootBoneTransform.GetLocation(),
+									//	RootBoneTransform.GetLocation() + (PredictedRootMotionQuat * RootBoneRotation).GetRightVector() * 100,
+									//	FColor::Orange, TEXT(""));
+
+									float YawToTargetOrientation = FRotator(DeltaToTargetOrientation).Yaw;
+
+									// pick the rotation direction that is the shortest path from the endpoint of the current animated rotation
+									if (PredictedRootMotionYaw - YawToTargetOrientation > 180)
+									{
+										YawToTargetOrientation += 360;
+									}
+									else if (YawToTargetOrientation - PredictedRootMotionYaw > 180)
+									{
+										YawToTargetOrientation -= 360;
+									}
+
+									float Ratio = YawToTargetOrientation / PredictedRootMotionYaw;
+
+									FRotator ThisFrameRootMotionRotation(ThisFrameRootMotionTransform.GetRotation());
+									ThisFrameRootMotionRotation.Yaw *= Ratio;
+
+									ThisFrameRootMotionTransform.SetRotation(FQuat(ThisFrameRootMotionRotation));
+									RootMotionProvider->OverrideRootMotion(ThisFrameRootMotionTransform, Keyframe->Get()->Attributes);
+
+									return;
+								}
+							}
+						}
+
+						if (SharedData->ProceduralTargetTime > InstanceData->DeltaTime)
+						{
+							DeltaToTargetOrientation = FQuat::Slerp(FQuat::Identity, DeltaToTargetOrientation, InstanceData->DeltaTime / SharedData->ProceduralTargetTime);
+						}
+
+						ThisFrameRootMotionTransform.SetRotation(FQuat::Slerp(ThisFrameRootMotionTransform.GetRotation(), DeltaToTargetOrientation, SharedData->Alpha));
+						RootMotionProvider->OverrideRootMotion(ThisFrameRootMotionTransform, Keyframe->Get()->Attributes);
+					}
+				}
+			}
+		}
+	}
+}
diff --git a/Engine/Plugins/Experimental/AnimNextWarping/Source/AnimNextWarping/Private/SteeringTrait.h b/Engine/Plugins/Experimental/AnimNextWarping/Source/AnimNextWarping/Private/SteeringTrait.h
new file mode 100644
--- /dev/null
+++ b/Engine/Plugins/Experimental/AnimNextWarping/Source/AnimNextWarping/Private/SteeringTrait.h
@@ -0,0 +1,149 @@
+// Copyright Epic Games, Inc. All Rights Reserved.
+
+#pragma once
+
+#include "CoreMinimal.h"
+#include "Kismet/KismetMathLibrary.h"
+#include "TraitCore/Trait.h"
+#include "TraitCore/TraitSharedData.h"
+#include "TraitInterfaces/IEvaluate.h"
+#include "TraitInterfaces/IUpdate.h"
+#include "SteeringTrait.generated.h"
+
+class UAnimSequence;
+
+/**
+ * Data needed to execute steering
+ * 
+ * Some steering data such as current anim asset / playback time is acquired via trait stack interfaces
+ */
+USTRUCT(meta = (DisplayName = "Steering"))
+struct FSteeringTraitSharedData : public FAnimNextTraitSharedData
+{
+	GENERATED_BODY()
+
+	// @TODO: This should belong in a parent class possibly? Discuss review, make JIRA
+	// Current strength of the skeletal control
+	UPROPERTY(EditAnywhere, Category = Alpha, meta = (PinShownByDefault))
+	float Alpha = 1.0f;
+
+	// The Orientation to steer towards
+	UPROPERTY(EditAnywhere, Category=Evaluation, meta=(PinShownByDefault))
+	FQuat TargetOrientation = FQuat::Identity;
+
+	// @TODO Temp / try to remove this. Shouldn't have to feed as argument
+	// Last root bone transform sampled
+	UPROPERTY(EditAnywhere, Category = Evaluation, meta = (PinShownByDefault))
+	FTransform RootBoneTransform = FTransform::Identity;
+
+	UPROPERTY(EditAnywhere, Category=Filtering)
+	bool bEnableTargetSmoothing = false;
+	
+	UPROPERTY(EditAnywhere, Category=Filtering)
+	float SmoothTargetStiffness = 300;
+	
+	UPROPERTY(EditAnywhere, Category=Filtering)
+	float SmoothTargetDamping = 1;
+	
+	// The number of seconds in the future before we should reach the TargetOrientation when play animations with no root motion rotation
+	UPROPERTY(EditAnywhere, Category=Evaluation, meta=(PinShownByDefault))
+	float ProceduralTargetTime = 0.2f;
+	
+	// The number of seconds in the future before we should reach the TargetOrientation when playing animations with root motion rotation
+	UPROPERTY(EditAnywhere, Category=Evaluation, meta=(PinShownByDefault))
+	float AnimatedTargetTime = 0.2f;
+
+	// If less than this number of degrees of rotation is found over TargetTime seconds in the CurrentAnimAsset, then rotation will be added linearly to reach the TargetOrientation
+	// Otherwise the root motion will be scaled to reach the TargetOrientation over TargetTime seconds
+	UPROPERTY(EditAnywhere, DisplayName=RootMotionAngleThreshold, Category=Evaluation)
+	float RootMotionThreshold = 1.0f;
+
+	// below this movement speed (based on the root motion in the animation) disable steering
+	UPROPERTY(EditAnywhere, Category=Evaluation)
+	float DisableSteeringBelowSpeed = 1.0f;
+
+	// Latent pin support boilerplate
+	#define TRAIT_LATENT_PROPERTIES_ENUMERATOR(GeneratorMacro) \
+		GeneratorMacro(Alpha) \
+		GeneratorMacro(TargetOrientation) \
+		GeneratorMacro(RootBoneTransform) \
+		GeneratorMacro(bEnableTargetSmoothing) \
+		GeneratorMacro(SmoothTargetStiffness) \
+		GeneratorMacro(SmoothTargetDamping) \
+		GeneratorMacro(ProceduralTargetTime) \
+		GeneratorMacro(AnimatedTargetTime) \
+		GeneratorMacro(RootMotionThreshold) \
+		GeneratorMacro(DisableSteeringBelowSpeed) \
+
+	GENERATE_TRAIT_LATENT_PROPERTIES(FSteeringTraitSharedData, TRAIT_LATENT_PROPERTIES_ENUMERATOR)
+	#undef TRAIT_LATENT_PROPERTIES_ENUMERATOR
+};
+
+/**
+ * Add procedural delta to root motion to match target orientations
+ * 
+ * Ex: If your anim only rotates say 45 deg, but you need to warp it to 60 deg to match gameplay input
+ */
+namespace UE::AnimNext
+{
+
+struct FSteeringTrait : FAdditiveTrait, IUpdate, IEvaluate
+{
+	// @TODO: Hash we use use FNV1A32. Update macro to mention that and also comment if we still need it (?)
+	DECLARE_ANIM_TRAIT(FSteeringTrait, 0x387fdf2a, FAdditiveTrait)
+
+	using FSharedData = FSteeringTraitSharedData;
+
+	// @TOOD: Memlayout this, after review, right now is organized to diff against old
+	struct FInstanceData : FTrait::FInstanceData
+	{
+		/** Delta in seconds between updates, populated during PreUpdate */
+		float DeltaTime = 0.f;
+
+		/** Target orientation per instance */
+		FQuat TargetOrientation = FQuat::Identity;
+
+		/** Current Anim Asset, we look ahead to determine how much steering is needed, populated during PostEvaluate */
+		TObjectPtr<const UAnimSequence> CurrentAnimAsset = nullptr;
+
+		/** Current Anim Asset time, used with ProceduralTargetTime to get future root motion, populated during PostEvaluate  */
+		float CurrentAnimAssetTime = 0;
+
+		// @TODO: Currently not used, for reinits
+		bool bResetFilter = true;
+
+		/** Target after a smoothing spring has been applied */
+		FQuat FilteredTarget = FQuat::Identity;
+		
+		/** Smoothing to apply to Target */
+		FQuaternionSpringState TargetSmoothingState = FQuaternionSpringState();
+
+		/** Last root bone transform sampled */
+		FTransform RootBoneTransform = FTransform::Identity;
+	};
+
+	// IUpdate impl 
+	virtual void PreUpdate(FUpdateTraversalContext& Context, const TTraitBinding<IUpdate>& Binding, const FTraitUpdateState& TraitState) const override;
+
+	// IEvaluate impl
+	virtual void PostEvaluate(FEvaluateTraversalContext& Context, const TTraitBinding<IEvaluate>& Binding) const override;
+};
+
+} // namespace UE::AnimNext
+
+/** Task to run Steering on VM */
+USTRUCT()
+struct FAnimNextSteeringTask : public FAnimNextEvaluationTask
+{
+	GENERATED_BODY()
+
+	DECLARE_ANIM_EVALUATION_TASK(FAnimNextSteeringTask)
+
+	static FAnimNextSteeringTask Make(UE::AnimNext::FSteeringTrait::FInstanceData* InstanceData, const UE::AnimNext::FSteeringTrait::FSharedData* SharedData);
+
+	// Task entry point
+	virtual void Execute(UE::AnimNext::FEvaluationVM& VM) const override;
+
+	UE::AnimNext::FSteeringTrait::FInstanceData* InstanceData = nullptr;
+	const UE::AnimNext::FSteeringTrait::FSharedData* SharedData = nullptr;
+};
\ No newline at end of file
