Fix float / int implict cast compare in steering trait

#tests NonUnity

diff --git a/Engine/Plugins/Experimental/AnimNextWarping/Source/AnimNextWarping/Private/SteeringTrait.cpp b/Engine/Plugins/Experimental/AnimNextWarping/Source/AnimNextWarping/Private/SteeringTrait.cpp
--- a/Engine/Plugins/Experimental/AnimNextWarping/Source/AnimNextWarping/Private/SteeringTrait.cpp
+++ b/Engine/Plugins/Experimental/AnimNextWarping/Source/AnimNextWarping/Private/SteeringTrait.cpp
@@ -170,7 +170,7 @@ void FAnimNextSteeringTask::Execute(UE::AnimNext::FEvaluationVM& VM) const
 
 						FQuat DeltaToTargetOrientation = RootBoneRotation.Inverse() * InstanceData->TargetOrientation;
 
-						if (SharedData->AnimatedTargetTime > 0)
+						if (SharedData->AnimatedTargetTime > 0.0f)
 						{
 							if (const UAnimSequenceBase* AnimSequence = InstanceData->CurrentAnimAsset)
 							{
