Change ControlRig dependency to RigVM for AnimNextWarping
- RigVM minimally needed for RigVM Registry.

#tests NonUnity, PF

diff --git a/Engine/Plugins/Experimental/AnimNextWarping/AnimNextWarping.uplugin b/Engine/Plugins/Experimental/AnimNextWarping/AnimNextWarping.uplugin
--- a/Engine/Plugins/Experimental/AnimNextWarping/AnimNextWarping.uplugin
+++ b/Engine/Plugins/Experimental/AnimNextWarping/AnimNextWarping.uplugin
@@ -26,7 +26,7 @@
 			"Enabled": true
 		},
 		{
-			"Name": "ControlRig",
+			"Name": "RigVM",
 			"Enabled": true
 		}
 	],
diff --git a/Engine/Plugins/Experimental/AnimNextWarping/Source/AnimNextWarping/AnimNextWarping.Build.cs b/Engine/Plugins/Experimental/AnimNextWarping/Source/AnimNextWarping/AnimNextWarping.Build.cs
--- a/Engine/Plugins/Experimental/AnimNextWarping/Source/AnimNextWarping/AnimNextWarping.Build.cs
+++ b/Engine/Plugins/Experimental/AnimNextWarping/Source/AnimNextWarping/AnimNextWarping.Build.cs
@@ -12,7 +12,6 @@ namespace UnrealBuildTool.Rules
 					"Core",
 					"CoreUObject",
 					"RigVM",
-					"ControlRig",
 					"Engine",
 					"AnimNext",
 					"AnimNextAnimGraph",
