Add IAttributeProvider trait interface

#tests PIE

diff --git a/Engine/Plugins/Experimental/AnimNextAnimGraph/Source/AnimNextAnimGraph/Internal/Traits/BlendStackTrait.h b/Engine/Plugins/Experimental/AnimNextAnimGraph/Source/AnimNextAnimGraph/Internal/Traits/BlendStackTrait.h
--- a/Engine/Plugins/Experimental/AnimNextAnimGraph/Source/AnimNextAnimGraph/Internal/Traits/BlendStackTrait.h
+++ b/Engine/Plugins/Experimental/AnimNextAnimGraph/Source/AnimNextAnimGraph/Internal/Traits/BlendStackTrait.h
@@ -11,6 +11,7 @@
 #include "TraitInterfaces/IGarbageCollection.h"
 #include "TraitInterfaces/IHierarchy.h"
 #include "TraitInterfaces/IUpdate.h"
+#include "TraitInterfaces/IAttributeProvider.h"
 #include "TraitInterfaces/ITimeline.h"
 #include "TraitInterfaces/ISmoothBlend.h"
 
@@ -85,7 +86,7 @@ namespace UE::AnimNext
 	 * 
 	 * Used as base trait in state machine-like setups where state machine is the additive trait that pushes graphs and blend settings to the blend stack.
 	 */
-	struct FBlendStackCoreTrait : FBaseTrait, IUpdateTraversal, IHierarchy, IDiscreteBlend, ISmoothBlend, ITimeline, IGarbageCollection, IBlendStack
+	struct FBlendStackCoreTrait : FBaseTrait, IUpdateTraversal, IHierarchy, IDiscreteBlend, ISmoothBlend, IAttributeProvider, ITimeline, IGarbageCollection, IBlendStack
 	{
 		DECLARE_ANIM_TRAIT(FBlendStackCoreTrait, 0xda7b7f8d, FBaseTrait)
 
@@ -153,6 +154,9 @@ namespace UE::AnimNext
 		virtual EAlphaBlendOption GetBlendType(FExecutionContext& Context, const TTraitBinding<ISmoothBlend>& Binding, int32 ChildIndex) const override;
 		virtual UCurveFloat* GetCustomBlendCurve(FExecutionContext& Context, const TTraitBinding<ISmoothBlend>& Binding, int32 ChildIndex) const override;
 
+		// IAttributeProvider impl
+		virtual FOnExtractRootMotionAttribute GetOnExtractRootMotionAttribute(FExecutionContext& Context, const TTraitBinding<IAttributeProvider>& Binding) const override;
+
 		// ITimeline impl
 		virtual FTimelineState GetState(FExecutionContext& Context, const TTraitBinding<ITimeline>& Binding) const override;
 
diff --git a/Engine/Plugins/Experimental/AnimNextAnimGraph/Source/AnimNextAnimGraph/Private/TraitInterfaces/IAttributeProvider.cpp b/Engine/Plugins/Experimental/AnimNextAnimGraph/Source/AnimNextAnimGraph/Private/TraitInterfaces/IAttributeProvider.cpp
new file mode 100644
--- /dev/null
+++ b/Engine/Plugins/Experimental/AnimNextAnimGraph/Source/AnimNextAnimGraph/Private/TraitInterfaces/IAttributeProvider.cpp
@@ -0,0 +1,34 @@
+// Copyright Epic Games, Inc. All Rights Reserved.
+
+#include "TraitInterfaces/IAttributeProvider.h"
+
+#include "TraitCore/ExecutionContext.h"
+
+namespace UE::AnimNext
+{
+	AUTO_REGISTER_ANIM_TRAIT_INTERFACE(IAttributeProvider)
+
+#if WITH_EDITOR
+	const FText& IAttributeProvider::GetDisplayName() const
+	{
+		static FText InterfaceName = NSLOCTEXT("TraitInterfaces", "TraitInterface_IAttributeProvider_Name", "AttributeProvider");
+		return InterfaceName;
+	}
+	const FText& IAttributeProvider::GetDisplayShortName() const
+	{
+		static FText InterfaceShortName = NSLOCTEXT("TraitInterfaces", "TraitInterface_IAttributeProvider_ShortName", "ATP");
+		return InterfaceShortName;
+	}
+#endif // WITH_EDITOR
+
+	FOnExtractRootMotionAttribute IAttributeProvider::GetOnExtractRootMotionAttribute(FExecutionContext& Context, const TTraitBinding<IAttributeProvider>& Binding) const
+	{
+		TTraitBinding<IAttributeProvider> SuperBinding;
+		if (Binding.GetStackInterfaceSuper(SuperBinding))
+		{
+			return SuperBinding.GetOnExtractRootMotionAttribute(Context);
+		}
+
+		return FOnExtractRootMotionAttribute();
+	}
+}
diff --git a/Engine/Plugins/Experimental/AnimNextAnimGraph/Source/AnimNextAnimGraph/Private/Traits/BlendStackTrait.cpp b/Engine/Plugins/Experimental/AnimNextAnimGraph/Source/AnimNextAnimGraph/Private/Traits/BlendStackTrait.cpp
--- a/Engine/Plugins/Experimental/AnimNextAnimGraph/Source/AnimNextAnimGraph/Private/Traits/BlendStackTrait.cpp
+++ b/Engine/Plugins/Experimental/AnimNextAnimGraph/Source/AnimNextAnimGraph/Private/Traits/BlendStackTrait.cpp
@@ -2,10 +2,12 @@
 
 #include "Traits/BlendStackTrait.h"
 
+#include "Animation/AnimSequence.h"
 #include "AnimNextAnimGraphSettings.h"
 #include "AnimNextDataInterfacePayload.h"
 #include "DataInterface/DataInterfaceStructAdapter.h"
 #include "Graph/AnimNextGraphInstance.h"
+#include "GraphInterfaces/AnimNextNativeDataInterface_AnimSequencePlayer.h"
 #include "TraitCore/NodeInstance.h"
 #include "TraitCore/ExecutionContext.h"
 #include "TraitInterfaces/IGraphFactory.h"
@@ -23,6 +25,7 @@ namespace UE::AnimNext
 		GeneratorMacro(IGarbageCollection) \
 		GeneratorMacro(IHierarchy) \
 		GeneratorMacro(ISmoothBlend) \
+		GeneratorMacro(IAttributeProvider) \
 		GeneratorMacro(ITimeline) \
 		GeneratorMacro(IUpdateTraversal) \
 		GeneratorMacro(IBlendStack) \
@@ -219,6 +222,40 @@ namespace UE::AnimNext
 		}
 	}
 
+	FOnExtractRootMotionAttribute FBlendStackCoreTrait::GetOnExtractRootMotionAttribute(FExecutionContext& Context, const TTraitBinding<IAttributeProvider>& Binding) const
+	{
+		TTraitBinding<IBlendStack> BlendStackTrait;
+		Binding.GetStackInterface(BlendStackTrait);
+
+		IBlendStack::FGraphRequestPtr ActiveGraphRequest;
+		BlendStackTrait.GetActiveGraph(Context, ActiveGraphRequest);
+
+		if (ActiveGraphRequest)
+		{
+			for (const FInstancedStruct& Payload : ActiveGraphRequest->GraphPayload.GetNativePayloads())
+			{
+				if (const FAnimNextNativeDataInterface_AnimSequencePlayer* PlayAnimPayload = Payload.GetPtr<FAnimNextNativeDataInterface_AnimSequencePlayer>())
+				{
+					// @TODO: Selecting the first active payload in a blend stack doesn't typically make sense, a wrapper trait is better.
+					// However this implementation is still useful for prototyping / debugging purposes. Still consider future removal.
+					if (PlayAnimPayload->AnimSequence)
+					{
+						auto ExtractRootMotionAttribute = [AnimSequence = PlayAnimPayload->AnimSequence](float StartTime, float DeltaTime, bool bAllowLooping)
+						{
+							// We do not check for lifetimes, assume the sequence is alive during pose list exectuion.
+							check(AnimSequence->IsValidLowLevel());
+							return AnimSequence->ExtractRootMotion(StartTime, DeltaTime, bAllowLooping);
+						};
+
+						return FOnExtractRootMotionAttribute::CreateLambda(ExtractRootMotionAttribute);
+					}
+				}
+			}
+		}
+
+		return FOnExtractRootMotionAttribute();
+	}
+
 	FTimelineState FBlendStackCoreTrait::GetState(FExecutionContext& Context, const TTraitBinding<ITimeline>& Binding) const
 	{
 		const FInstanceData* InstanceData = Binding.GetInstanceData<FInstanceData>();
diff --git a/Engine/Plugins/Experimental/AnimNextAnimGraph/Source/AnimNextAnimGraph/Private/Traits/SequencePlayer.cpp b/Engine/Plugins/Experimental/AnimNextAnimGraph/Source/AnimNextAnimGraph/Private/Traits/SequencePlayer.cpp
--- a/Engine/Plugins/Experimental/AnimNextAnimGraph/Source/AnimNextAnimGraph/Private/Traits/SequencePlayer.cpp
+++ b/Engine/Plugins/Experimental/AnimNextAnimGraph/Source/AnimNextAnimGraph/Private/Traits/SequencePlayer.cpp
@@ -2,6 +2,7 @@
 
 #include "Traits/SequencePlayer.h"
 
+#include "Animation/AnimSequenceBase.h"
 #include "AnimationRuntime.h"
 #include "TraitCore/ExecutionContext.h"
 #include "EvaluationVM/Tasks/PushAnimSequenceKeyframe.h"
@@ -13,6 +14,7 @@ namespace UE::AnimNext
 	// Trait implementation boilerplate
 	#define TRAIT_INTERFACE_ENUMERATOR(GeneratorMacro) \
 		GeneratorMacro(IEvaluate) \
+		GeneratorMacro(IAttributeProvider) \
 		GeneratorMacro(ITimeline) \
 		GeneratorMacro(IUpdate) \
 		GeneratorMacro(IGarbageCollection) \
@@ -52,6 +54,25 @@ namespace UE::AnimNext
 		Context.AppendTask(Task);
 	}
 
+	FOnExtractRootMotionAttribute FSequencePlayerTrait::GetOnExtractRootMotionAttribute(FExecutionContext& Context, const TTraitBinding<IAttributeProvider>& Binding) const
+	{
+		const FInstanceData* InstanceData = Binding.GetInstanceData<FInstanceData>();
+
+		if (InstanceData->AnimSequence)
+		{
+			auto ExtractRootMotionAttribute = [AnimSequence = InstanceData->AnimSequence](float StartTime, float DeltaTime, bool bAllowLooping)
+			{
+				// We do not check for lifetimes, assume the sequence is alive during pose list exectuion.
+				check(AnimSequence->IsValidLowLevel());
+				return AnimSequence->ExtractRootMotion(StartTime, DeltaTime, bAllowLooping);
+			};
+
+			return FOnExtractRootMotionAttribute::CreateLambda(ExtractRootMotionAttribute);
+		}
+
+		return FOnExtractRootMotionAttribute();
+	}
+
 	void FSequencePlayerTrait::GetSyncMarkers(FExecutionContext& Context, const TTraitBinding<ITimeline>& Binding, FTimelineSyncMarkerArray& OutSyncMarkers) const
 	{
 		const FInstanceData* InstanceData = Binding.GetInstanceData<FInstanceData>();
diff --git a/Engine/Plugins/Experimental/AnimNextAnimGraph/Source/AnimNextAnimGraph/Private/Traits/SequencePlayer.h b/Engine/Plugins/Experimental/AnimNextAnimGraph/Source/AnimNextAnimGraph/Private/Traits/SequencePlayer.h
--- a/Engine/Plugins/Experimental/AnimNextAnimGraph/Source/AnimNextAnimGraph/Private/Traits/SequencePlayer.h
+++ b/Engine/Plugins/Experimental/AnimNextAnimGraph/Source/AnimNextAnimGraph/Private/Traits/SequencePlayer.h
@@ -6,6 +6,7 @@
 #include "TraitCore/Trait.h"
 #include "TraitInterfaces/IEvaluate.h"
 #include "TraitInterfaces/IGarbageCollection.h"
+#include "TraitInterfaces/IAttributeProvider.h"
 #include "TraitInterfaces/ITimeline.h"
 #include "TraitInterfaces/IUpdate.h"
 #include "Animation/AnimSequence.h"
@@ -51,7 +52,7 @@ namespace UE::AnimNext
 	 * 
 	 * A trait that can play an animation sequence.
 	 */
-	struct FSequencePlayerTrait : FBaseTrait, IEvaluate, ITimeline, IUpdate, IGarbageCollection
+	struct FSequencePlayerTrait : FBaseTrait, IEvaluate, IAttributeProvider, ITimeline, IUpdate, IGarbageCollection
 	{
 		DECLARE_ANIM_TRAIT(FSequencePlayerTrait, 0x7a0dc157, FBaseTrait)
 
@@ -73,6 +74,9 @@ namespace UE::AnimNext
 		// IEvaluate impl
 		virtual void PreEvaluate(FEvaluateTraversalContext& Context, const TTraitBinding<IEvaluate>& Binding) const override;
 
+		// IAttributeProvider impl
+		virtual FOnExtractRootMotionAttribute GetOnExtractRootMotionAttribute(FExecutionContext& Context, const TTraitBinding<IAttributeProvider>& Binding) const override;
+
 		// ITimeline impl
 		virtual void GetSyncMarkers(FExecutionContext& Context, const TTraitBinding<ITimeline>& Binding, FTimelineSyncMarkerArray& OutSyncMarkers) const override;
 		virtual FTimelineState GetState(FExecutionContext& Context, const TTraitBinding<ITimeline>& Binding) const override;
diff --git a/Engine/Plugins/Experimental/AnimNextAnimGraph/Source/AnimNextAnimGraph/Public/TraitCore/Trait.h b/Engine/Plugins/Experimental/AnimNextAnimGraph/Source/AnimNextAnimGraph/Public/TraitCore/Trait.h
--- a/Engine/Plugins/Experimental/AnimNextAnimGraph/Source/AnimNextAnimGraph/Public/TraitCore/Trait.h
+++ b/Engine/Plugins/Experimental/AnimNextAnimGraph/Source/AnimNextAnimGraph/Public/TraitCore/Trait.h
@@ -20,6 +20,7 @@ class FArchive;
 
 // Helper macros
 // In the trait class declaration, this macro declares the Super alias and base functions we override
+// We use use FNV1A32 to hash entire TraitName, including 'F' part (Currently must be manually done). 
 #define DECLARE_ANIM_TRAIT(TraitName, TraitNameHash, SuperTraitName) \
 	using TraitSuper = SuperTraitName; \
 	/* FTrait impl */ \
diff --git a/Engine/Plugins/Experimental/AnimNextAnimGraph/Source/AnimNextAnimGraph/Public/TraitInterfaces/IAttributeProvider.h b/Engine/Plugins/Experimental/AnimNextAnimGraph/Source/AnimNextAnimGraph/Public/TraitInterfaces/IAttributeProvider.h
new file mode 100644
--- /dev/null
+++ b/Engine/Plugins/Experimental/AnimNextAnimGraph/Source/AnimNextAnimGraph/Public/TraitInterfaces/IAttributeProvider.h
@@ -0,0 +1,52 @@
+// Copyright Epic Games, Inc. All Rights Reserved.
+
+#pragma once
+
+#include "CoreMinimal.h"
+#include "TraitCore/ITraitInterface.h"
+#include "TraitCore/TraitBinding.h"
+
+class UAnimSequenceBase;
+
+namespace UE::AnimNext
+{
+	/**
+	 * Root motion attribute sampling may occur during execution (Ex: Conditionally sampling attributes based on active root motion).
+	 * This signature is used to support that deferred execution, each implementing trait is responsible for ensuring callback lifetimes.
+	 */
+	DECLARE_DELEGATE_RetVal_ThreeParams(FTransform, FOnExtractRootMotionAttribute, float /* StartTime */, float /* DeltaTime */, bool /* bAllowLooping */);
+
+	/**
+	 * IAttributeProvider
+	 *
+	 * This interface provides access to animation attributes such as root motion
+	 */
+	struct ANIMNEXTANIMGRAPH_API IAttributeProvider : ITraitInterface
+	{
+		DECLARE_ANIM_TRAIT_INTERFACE(IAttributeProvider, 0x698139d1)
+
+		// Returns a callback that can be used to extract root motion at a later time
+		virtual FOnExtractRootMotionAttribute GetOnExtractRootMotionAttribute(FExecutionContext& Context, const TTraitBinding<IAttributeProvider>& Binding) const;
+
+#if WITH_EDITOR
+		virtual const FText& GetDisplayName() const override;
+		virtual const FText& GetDisplayShortName() const override;
+#endif // WITH_EDITOR
+	};
+
+	/**
+	 * Specialization for trait binding.
+	 */
+	template<>
+	struct TTraitBinding<IAttributeProvider> : FTraitBinding
+	{
+		// @see IAttributeProvider::GetOnExtractRootMotionAttribute
+		FOnExtractRootMotionAttribute GetOnExtractRootMotionAttribute(FExecutionContext& Context) const
+		{
+			return GetInterface()->GetOnExtractRootMotionAttribute(Context, *this);
+		}
+
+	protected:
+		const IAttributeProvider* GetInterface() const { return GetInterfaceTyped<IAttributeProvider>(); }
+	};
+}
diff --git a/Engine/Plugins/Experimental/AnimNextWarping/Source/AnimNextWarping/Private/SteeringTrait.cpp b/Engine/Plugins/Experimental/AnimNextWarping/Source/AnimNextWarping/Private/SteeringTrait.cpp
--- a/Engine/Plugins/Experimental/AnimNextWarping/Source/AnimNextWarping/Private/SteeringTrait.cpp
+++ b/Engine/Plugins/Experimental/AnimNextWarping/Source/AnimNextWarping/Private/SteeringTrait.cpp
@@ -5,14 +5,14 @@
 #include "AnimNextAnimGraphSettings.h"
 #include "AnimNextWarpingLog.h"
 #include "Animation/AnimRootMotionProvider.h"
-#include "Animation/AnimSequence.h"
+#include "Animation/AnimSequenceBase.h"
 #include "DataInterface/DataInterfaceStructAdapter.h"
 #include "EvaluationVM/EvaluationVM.h"
 #include "EvaluationVM/KeyframeState.h"
-#include "GraphInterfaces/AnimNextNativeDataInterface_AnimSequencePlayer.h"
-#include "TraitInterfaces/IBlendStack.h"
+#include "TraitInterfaces/IAttributeProvider.h"
 #include "TraitInterfaces/IGraphFactory.h"
 #include "TraitInterfaces/ITimeline.h"
+#include "VisualLogger/VisualLogger.h"
 
 ////////////////////////////////////////////////////////////////////////////////////////////////////////////////////////
 // FSteeringTrait
@@ -25,8 +25,12 @@ namespace UE::AnimNext
 		GeneratorMacro(IUpdate) \
 		GeneratorMacro(IEvaluate) \
 
+	// Trait required interfaces implementation boilerplate
+	#define TRAIT_REQUIRED_INTERFACE_ENUMERATOR(GeneratorMacroRequired) \
+		GeneratorMacroRequired(IAttributeProvider) \
+
 	// Trait implementation boilerplate
-	GENERATE_ANIM_TRAIT_IMPLEMENTATION(FSteeringTrait, TRAIT_INTERFACE_ENUMERATOR, NULL_ANIM_TRAIT_INTERFACE_ENUMERATOR, NULL_ANIM_TRAIT_EVENT_ENUMERATOR)
+	GENERATE_ANIM_TRAIT_IMPLEMENTATION(FSteeringTrait, TRAIT_INTERFACE_ENUMERATOR, TRAIT_REQUIRED_INTERFACE_ENUMERATOR, NULL_ANIM_TRAIT_EVENT_ENUMERATOR)
 	#undef TRAIT_INTERFACE_ENUMERATOR
 
 	void FSteeringTrait::PreUpdate(FUpdateTraversalContext& Context, const TTraitBinding<IUpdate>& Binding, const FTraitUpdateState& TraitState) const
@@ -51,37 +55,14 @@ namespace UE::AnimNext
 
 		// Get Current Anim Asset
 		{
-			TTraitBinding<IBlendStack> BlendStackTrait;
-			if (!Binding.GetStackInterface(BlendStackTrait))
+			TTraitBinding<IAttributeProvider> AttributeTrait;
+			if (!Binding.GetStackInterface(AttributeTrait))
 			{
-				UE_LOG(LogAnimNextWarping, Error, TEXT("FSteeringTrait::PostEvaluate, missing IBlendStack"));
+				UE_LOG(LogAnimNextWarping, Error, TEXT("FSteeringTrait::PostEvaluate, missing IAttributeProvider"));
 				return;
 			}
 
-			// @TODO: See if we want this as a trait interface (Some trait to get anim asset). 
-			auto GetCurrentAnimAsset = [&]()
-			{
-				TObjectPtr<const UAnimSequence> Result = nullptr;
-
-				IBlendStack::FGraphRequestPtr ActiveGraphRequest;
-				BlendStackTrait.GetActiveGraph(Context, ActiveGraphRequest);
-
-				if (ActiveGraphRequest)
-				{
-					for (const FInstancedStruct& Payload : ActiveGraphRequest->GraphPayload.GetNativePayloads())
-					{
-						if (const FAnimNextNativeDataInterface_AnimSequencePlayer* PlayAnimPayload = Payload.GetPtr<FAnimNextNativeDataInterface_AnimSequencePlayer>())
-						{
-							Result = PlayAnimPayload->AnimSequence;
-							return Result;
-						}
-					}
-				}
-
-				return Result;
-			};
-
-			InstanceData->CurrentAnimAsset = GetCurrentAnimAsset();
+			InstanceData->OnExtractRootMotionAttribute = AttributeTrait.GetOnExtractRootMotionAttribute(Context);
 		}
 
 		// Get Current Anim Asset Time
@@ -96,11 +77,13 @@ namespace UE::AnimNext
 			InstanceData->CurrentAnimAssetTime = TimelineTrait.GetState(Context).GetPosition();
 		}
 
-		// Update target orientation / root motion
-		{
-			InstanceData->TargetOrientation = SharedData->GetTargetOrientation(Binding);
-			InstanceData->RootBoneTransform = SharedData->GetRootBoneTransform(Binding);
-		}
+		// Update target orientation, root bone transform, & other properties 
+		InstanceData->TargetOrientation = SharedData->GetTargetOrientation(Binding);
+		InstanceData->RootBoneTransform = SharedData->GetRootBoneTransform(Binding);
+
+#if ENABLE_ANIM_DEBUG 
+		InstanceData->HostObject = Context.GetHostObject();
+#endif // ENABLE_ANIM_DEBUG 
 
 		Context.AppendTask(FAnimNextSteeringTask::Make(InstanceData, SharedData));
 	}
@@ -158,32 +141,36 @@ void FAnimNextSteeringTask::Execute(UE::AnimNext::FEvaluationVM& VM) const
 					{
 						FQuat RootBoneRotation = InstanceData->RootBoneTransform.GetRotation();
 
-						//UE_VLOG_ARROW(Output.AnimInstanceProxy->GetAnimInstanceObject(), "Steering", Display,
-						//	RootBoneTransform.GetLocation(),
-						//	RootBoneTransform.GetLocation() + RootBoneRotation.GetRightVector() * 90,
-						//	FColor::Green, TEXT(""));
+#if ENABLE_ANIM_DEBUG
+						UE_VLOG_ARROW(InstanceData->HostObject, "Steering", Display,
+							InstanceData->RootBoneTransform.GetLocation(),
+							InstanceData->RootBoneTransform.GetLocation() + RootBoneRotation.GetRightVector() * 90,
+							FColor::Green, TEXT(""));
 
-						//UE_VLOG_ARROW(Output.AnimInstanceProxy->GetAnimInstanceObject(), "Steering", Display,
-						//	RootBoneTransform.GetLocation(),
-						//	RootBoneTransform.GetLocation() + InstanceData->TargetOrientation.GetRightVector() * 100,
-						//	FColor::Blue, TEXT(""));
+						UE_VLOG_ARROW(InstanceData->HostObject, "Steering", Display,
+							InstanceData->RootBoneTransform.GetLocation(),
+							InstanceData->RootBoneTransform.GetLocation() + InstanceData->TargetOrientation.GetRightVector() * 100,
+							FColor::Blue, TEXT(""));
+#endif // ENABLE_ANIM_DEBUG 
 
 						FQuat DeltaToTargetOrientation = RootBoneRotation.Inverse() * InstanceData->TargetOrientation;
 
 						if (SharedData->AnimatedTargetTime > 0.0f)
 						{
-							if (const UAnimSequenceBase* AnimSequence = InstanceData->CurrentAnimAsset)
+							if (InstanceData->OnExtractRootMotionAttribute.IsBound())
 							{
-								FTransform PredictedRootMotionTransform = AnimSequence->ExtractRootMotion(InstanceData->CurrentAnimAssetTime, SharedData->AnimatedTargetTime, true);
+								FTransform PredictedRootMotionTransform = InstanceData->OnExtractRootMotionAttribute.Execute(InstanceData->CurrentAnimAssetTime, SharedData->AnimatedTargetTime, true);
 								FQuat PredictedRootMotionQuat = PredictedRootMotionTransform.GetRotation();
 								float PredictedRootMotionYaw = FRotator(PredictedRootMotionQuat).Yaw;
 
 								if (fabs(PredictedRootMotionYaw) > SharedData->RootMotionThreshold)
 								{
-									//UE_VLOG_ARROW(Output.AnimInstanceProxy->GetAnimInstanceObject(), "Steering", Display,
-									//	RootBoneTransform.GetLocation(),
-									//	RootBoneTransform.GetLocation() + (PredictedRootMotionQuat * RootBoneRotation).GetRightVector() * 100,
-									//	FColor::Orange, TEXT(""));
+#if ENABLE_ANIM_DEBUG
+									UE_VLOG_ARROW(InstanceData->HostObject, "Steering", Display,
+										InstanceData->RootBoneTransform.GetLocation(),
+										InstanceData->RootBoneTransform.GetLocation() + (PredictedRootMotionQuat * RootBoneRotation).GetRightVector() * 100,
+										FColor::Orange, TEXT(""));
+#endif // ENABLE_ANIM_DEBUG 
 
 									float YawToTargetOrientation = FRotator(DeltaToTargetOrientation).Yaw;
 
diff --git a/Engine/Plugins/Experimental/AnimNextWarping/Source/AnimNextWarping/Private/SteeringTrait.h b/Engine/Plugins/Experimental/AnimNextWarping/Source/AnimNextWarping/Private/SteeringTrait.h
--- a/Engine/Plugins/Experimental/AnimNextWarping/Source/AnimNextWarping/Private/SteeringTrait.h
+++ b/Engine/Plugins/Experimental/AnimNextWarping/Source/AnimNextWarping/Private/SteeringTrait.h
@@ -3,15 +3,15 @@
 #pragma once
 
 #include "CoreMinimal.h"
+#include "EngineDefines.h"
 #include "Kismet/KismetMathLibrary.h"
 #include "TraitCore/Trait.h"
 #include "TraitCore/TraitSharedData.h"
+#include "TraitInterfaces/IAttributeProvider.h"
 #include "TraitInterfaces/IEvaluate.h"
 #include "TraitInterfaces/IUpdate.h"
 #include "SteeringTrait.generated.h"
 
-class UAnimSequence;
-
 /**
  * Data needed to execute steering
  * 
@@ -89,28 +89,23 @@ namespace UE::AnimNext
 
 struct FSteeringTrait : FAdditiveTrait, IUpdate, IEvaluate
 {
-	// @TODO: Hash we use use FNV1A32. Update macro to mention that and also comment if we still need it (?)
 	DECLARE_ANIM_TRAIT(FSteeringTrait, 0x387fdf2a, FAdditiveTrait)
 
 	using FSharedData = FSteeringTraitSharedData;
 
-	// @TOOD: Memlayout this, after review, right now is organized to diff against old
 	struct FInstanceData : FTrait::FInstanceData
 	{
-		/** Delta in seconds between updates, populated during PreUpdate */
-		float DeltaTime = 0.f;
-
 		/** Target orientation per instance */
 		FQuat TargetOrientation = FQuat::Identity;
 
-		/** Current Anim Asset, we look ahead to determine how much steering is needed, populated during PostEvaluate */
-		TObjectPtr<const UAnimSequence> CurrentAnimAsset = nullptr;
+		/** Callback provided by attribute trait on stack to evaluate root motion at a later time */
+		FOnExtractRootMotionAttribute OnExtractRootMotionAttribute = FOnExtractRootMotionAttribute();
 
 		/** Current Anim Asset time, used with ProceduralTargetTime to get future root motion, populated during PostEvaluate  */
 		float CurrentAnimAssetTime = 0;
 
-		// @TODO: Currently not used, for reinits
-		bool bResetFilter = true;
+		/** Delta in seconds between updates, populated during PreUpdate */
+		float DeltaTime = 0.f;
 
 		/** Target after a smoothing spring has been applied */
 		FQuat FilteredTarget = FQuat::Identity;
@@ -120,6 +115,14 @@ struct FSteeringTrait : FAdditiveTrait, IUpdate, IEvaluate
 
 		/** Last root bone transform sampled */
 		FTransform RootBoneTransform = FTransform::Identity;
+
+		// @TODO: Currently not used, for reinits
+		bool bResetFilter = true;
+
+#if ENABLE_ANIM_DEBUG 
+		/** Debug Object for VisualLogger */
+		TObjectPtr<const UObject> HostObject = nullptr;
+#endif // ENABLE_ANIM_DEBUG 
 	};
 
 	// IUpdate impl 
