Warping - WarpTest animnode and trait

diff --git a/Engine/Plugins/Animation/AnimationWarping/Source/Editor/Private/AnimGraph/AnimGraphNode_WarpTest.cpp b/Engine/Plugins/Animation/AnimationWarping/Source/Editor/Private/AnimGraph/AnimGraphNode_WarpTest.cpp
new file mode 100644
--- /dev/null
+++ b/Engine/Plugins/Animation/AnimationWarping/Source/Editor/Private/AnimGraph/AnimGraphNode_WarpTest.cpp
@@ -0,0 +1,26 @@
+// Copyright Epic Games, Inc. All Rights Reserved.
+
+#include "AnimGraph/AnimGraphNode_WarpTest.h"
+
+#define LOCTEXT_NAMESPACE "AnimationWarping"
+
+UAnimGraphNode_WarpTest::UAnimGraphNode_WarpTest(const FObjectInitializer& ObjectInitializer) : Super(ObjectInitializer)
+{
+}
+
+FText UAnimGraphNode_WarpTest::GetTooltipText() const
+{
+	return LOCTEXT("WarpTestTooltip", "Warp Test Node");
+}
+
+FText UAnimGraphNode_WarpTest::GetNodeTitle(ENodeTitleType::Type TitleType) const
+{
+	return LOCTEXT("WarpTest", "Warp Test Node");
+}
+
+FLinearColor UAnimGraphNode_WarpTest::GetNodeTitleColor() const
+{
+	return FLinearColor(FColor(153.f, 0.f, 0.f));
+}
+
+#undef LOCTEXT_NAMESPACE
diff --git a/Engine/Plugins/Animation/AnimationWarping/Source/Editor/Public/AnimGraph/AnimGraphNode_WarpTest.h b/Engine/Plugins/Animation/AnimationWarping/Source/Editor/Public/AnimGraph/AnimGraphNode_WarpTest.h
new file mode 100644
--- /dev/null
+++ b/Engine/Plugins/Animation/AnimationWarping/Source/Editor/Public/AnimGraph/AnimGraphNode_WarpTest.h
@@ -0,0 +1,25 @@
+// Copyright Epic Games, Inc. All Rights Reserved.
+
+#pragma once
+
+#include "AnimGraphNode_Base.h"
+#include "BoneControllers/AnimNode_WarpTest.h"
+#include "AnimGraphNode_WarpTest.generated.h"
+
+namespace ENodeTitleType { enum Type : int; }
+
+UCLASS(Experimental)
+class ANIMATIONWARPINGEDITOR_API UAnimGraphNode_WarpTest : public UAnimGraphNode_Base
+{
+	GENERATED_UCLASS_BODY()
+
+	UPROPERTY(EditAnywhere, Category = Settings)
+	FAnimNode_WarpTest Node;
+
+public:
+	// UEdGraphNode interface
+	virtual FText GetNodeTitle(ENodeTitleType::Type TitleType) const override;
+	virtual FLinearColor GetNodeTitleColor() const override;
+	virtual FText GetTooltipText() const override;
+	// End of UEdGraphNode interface
+};
diff --git a/Engine/Plugins/Animation/AnimationWarping/Source/Runtime/Private/BoneControllers/AnimNode_WarpTest.cpp b/Engine/Plugins/Animation/AnimationWarping/Source/Runtime/Private/BoneControllers/AnimNode_WarpTest.cpp
new file mode 100644
--- /dev/null
+++ b/Engine/Plugins/Animation/AnimationWarping/Source/Runtime/Private/BoneControllers/AnimNode_WarpTest.cpp
@@ -0,0 +1,84 @@
+// Copyright Epic Games, Inc. All Rights Reserved.
+
+#include "BoneControllers/AnimNode_WarpTest.h"
+#include "Animation/AnimInstanceProxy.h"
+#include "Animation/AnimRootMotionProvider.h"
+#include "BoneControllers/AnimNode_OffsetRootBone.h"
+#include "VisualLogger/VisualLogger.h"
+
+void FAnimNode_WarpTest::Initialize_AnyThread(const FAnimationInitializeContext& Context)
+{
+	Super::Initialize_AnyThread(Context);
+	Source.Initialize(Context);
+}
+
+void FAnimNode_WarpTest::CacheBones_AnyThread(const FAnimationCacheBonesContext& Context)
+{
+	Super::CacheBones_AnyThread(Context);
+	Source.CacheBones(Context);
+}
+
+void FAnimNode_WarpTest::Update_AnyThread(const FAnimationUpdateContext& Context)
+{
+	Super::Update_AnyThread(Context);
+	GetEvaluateGraphExposedInputs().Execute(Context);
+	Source.Update(Context);
+
+	if (Transforms.IsEmpty())
+	{
+		CurrentTransformIndex = 0;
+		CurrentTime = 0.f;
+	}
+	else
+	{
+		if (CurrentTransformIndex >= Transforms.Num())
+		{
+			CurrentTransformIndex = 0;
+		}
+
+		CurrentTime += Context.GetDeltaTime();
+		if (CurrentTime > SecondsToWait)
+		{
+			CurrentTime -= SecondsToWait;
+			CurrentTransformIndex = (CurrentTransformIndex + 1) % Transforms.Num();
+		}
+	}
+
+	if (UE::AnimationWarping::FRootOffsetProvider* RootOffsetProvider = Context.GetMessage<UE::AnimationWarping::FRootOffsetProvider>())
+	{
+		ComponentTransform = RootOffsetProvider->GetRootTransform();
+	}
+	else
+	{
+		ComponentTransform = Context.AnimInstanceProxy->GetComponentTransform();
+	}
+}
+
+void FAnimNode_WarpTest::Evaluate_AnyThread(FPoseContext& Output)
+{
+	Super::Evaluate_AnyThread(Output);
+	Source.Evaluate(Output);
+	
+	if (!Transforms.IsEmpty())
+	{
+		const FTransform& WarpTo = Transforms[CurrentTransformIndex];
+		const FTransform RootMotion = WarpTo.GetRelativeTransform(ComponentTransform);
+
+		const UE::Anim::IAnimRootMotionProvider* RootMotionProvider = UE::Anim::IAnimRootMotionProvider::Get();
+		RootMotionProvider->OverrideRootMotion(RootMotion, Output.CustomAttributes);
+
+#if ENABLE_VISUAL_LOG
+		if (FVisualLogger::IsRecording())
+		{
+			UObject* AnimInstanceObject = Output.AnimInstanceProxy->GetAnimInstanceObject();
+
+			static const TCHAR* LogName = TEXT("WarpTestNode");
+			UE_VLOG_SEGMENT_THICK(AnimInstanceObject, LogName, Display, WarpTo.GetLocation(), WarpTo.GetRotation().GetAxisX() * 100.f + WarpTo.GetLocation(), FColor::Red, 1, TEXT(""));
+			UE_VLOG_SEGMENT_THICK(AnimInstanceObject, LogName, Display, WarpTo.GetLocation(), WarpTo.GetRotation().GetAxisY() * 100.f + WarpTo.GetLocation(), FColor::Blue, 1, TEXT(""));
+
+			UE_VLOG_SEGMENT_THICK(AnimInstanceObject, LogName, Display, ComponentTransform.GetLocation(), ComponentTransform.GetRotation().GetAxisX() * 80.f + ComponentTransform.GetLocation(), FColor::Black, 1, TEXT(""));
+			UE_VLOG_SEGMENT_THICK(AnimInstanceObject, LogName, Display, ComponentTransform.GetLocation(), ComponentTransform.GetRotation().GetAxisY() * 80.f + ComponentTransform.GetLocation(), FColor::Green, 1, TEXT(""));
+		}
+#endif // ENABLE_VISUAL_LOG
+	}
+}
\ No newline at end of file
diff --git a/Engine/Plugins/Animation/AnimationWarping/Source/Runtime/Public/BoneControllers/AnimNode_WarpTest.h b/Engine/Plugins/Animation/AnimationWarping/Source/Runtime/Public/BoneControllers/AnimNode_WarpTest.h
new file mode 100644
--- /dev/null
+++ b/Engine/Plugins/Animation/AnimationWarping/Source/Runtime/Public/BoneControllers/AnimNode_WarpTest.h
@@ -0,0 +1,36 @@
+// Copyright Epic Games, Inc. All Rights Reserved.
+
+#pragma once
+
+#include "Animation/AnimNodeBase.h"
+#include "AnimNode_WarpTest.generated.h"
+
+USTRUCT(Experimental, BlueprintInternalUseOnly)
+struct ANIMATIONWARPINGRUNTIME_API FAnimNode_WarpTest : public FAnimNode_Base
+{
+	GENERATED_BODY();
+
+	UPROPERTY(EditAnywhere, BlueprintReadWrite, Category = Links, meta = (DisplayPriority = 0))
+	FPoseLink Source;
+
+	// the node will warp the character looping between Transforms[i] choosing the next one every SecondsToWait
+	UPROPERTY(EditAnywhere, Category = Evaluation, meta = (PinShownByDefault))
+	TArray<FTransform> Transforms;
+
+	// every SecondsToWait we warp to the next Transforms[i]
+	UPROPERTY(EditAnywhere, Category = Evaluation, meta = (PinShownByDefault))
+	float SecondsToWait = 1.f;
+
+public:
+	// FAnimNode_Base interface
+	virtual void Update_AnyThread(const FAnimationUpdateContext& Context) override;
+	virtual void Evaluate_AnyThread(FPoseContext& Output) override;
+	virtual void Initialize_AnyThread(const FAnimationInitializeContext& Context) override;
+	virtual void CacheBones_AnyThread(const FAnimationCacheBonesContext& Context) override;
+	// End of FAnimNode_Base interface
+
+	// ComponentTransform represents the previous frame component transform
+	FTransform ComponentTransform = FTransform::Identity;
+	float CurrentTime = 0.f;
+	int32 CurrentTransformIndex = 0;
+};
diff --git a/Engine/Plugins/Experimental/AnimNextWarping/Source/AnimNextWarping/Private/WarpTestTrait.cpp b/Engine/Plugins/Experimental/AnimNextWarping/Source/AnimNextWarping/Private/WarpTestTrait.cpp
new file mode 100644
--- /dev/null
+++ b/Engine/Plugins/Experimental/AnimNextWarping/Source/AnimNextWarping/Private/WarpTestTrait.cpp
@@ -0,0 +1,135 @@
+// Copyright Epic Games, Inc. All Rights Reserved.
+
+#include "WarpTestTrait.h"
+#include "Animation/AnimRootMotionProvider.h"
+#include "AnimNextWarpingLog.h"
+#include "Component/AnimNextComponent.h"
+#include "Components/SkeletalMeshComponent.h"
+#include "EvaluationVM/EvaluationVM.h"
+#include "Module/AnimNextModuleInstance.h"
+#include "VisualLogger/VisualLogger.h"
+
+namespace UE::AnimNext
+{
+	AUTO_REGISTER_ANIM_TRAIT(FWarpTestTrait)
+
+	#define TRAIT_INTERFACE_ENUMERATOR(GeneratorMacro) \
+		GeneratorMacro(IUpdate) \
+		GeneratorMacro(IEvaluate) \
+
+	GENERATE_ANIM_TRAIT_IMPLEMENTATION(FWarpTestTrait, TRAIT_INTERFACE_ENUMERATOR, NULL_ANIM_TRAIT_INTERFACE_ENUMERATOR, NULL_ANIM_TRAIT_EVENT_ENUMERATOR)
+
+	#undef TRAIT_INTERFACE_ENUMERATOR
+
+	void FWarpTestTrait::PreUpdate(FUpdateTraversalContext& Context, const TTraitBinding<IUpdate>& Binding, const FTraitUpdateState& TraitState) const
+	{
+		const FSharedData* SharedData = Binding.GetSharedData<FSharedData>();
+		check(SharedData);
+
+		FInstanceData* InstanceData = Binding.GetInstanceData<FInstanceData>();
+
+		check(InstanceData);
+
+		const float SecondsToWait = SharedData->GetSecondsToWait(Binding);
+		const TArray<FTransform>& Transforms = SharedData->GetTransforms(Binding);
+
+		if (Transforms.IsEmpty())
+		{
+			InstanceData->CurrentTransformIndex = 0;
+			InstanceData->CurrentTime = 0.f;
+		}
+		else
+		{
+			if (InstanceData->CurrentTransformIndex >= Transforms.Num())
+			{
+				InstanceData->CurrentTransformIndex = 0;
+			}
+
+			InstanceData->CurrentTime += TraitState.GetDeltaTime();
+			if (InstanceData->CurrentTime > SecondsToWait)
+			{
+				InstanceData->CurrentTime -= SecondsToWait;
+				InstanceData->CurrentTransformIndex = (InstanceData->CurrentTransformIndex + 1) % Transforms.Num();
+			}
+
+			if (InstanceData->CurrentTransformIndex < Transforms.Num())
+			{
+				InstanceData->WarpTo = Transforms[InstanceData->CurrentTransformIndex];
+			}
+		}
+
+		IUpdate::PreUpdate(Context, Binding, TraitState);
+	}
+
+	void FWarpTestTrait::PostEvaluate(FEvaluateTraversalContext& Context, const TTraitBinding<IEvaluate>& Binding) const
+	{
+		IEvaluate::PostEvaluate(Context, Binding);
+
+		const FSharedData* SharedData = Binding.GetSharedData<FSharedData>();
+		check(SharedData);
+
+		FInstanceData* InstanceData = Binding.GetInstanceData<FInstanceData>();
+		check(InstanceData);
+
+		// hacky way to retrieve the mesh transform until we find a better way
+		if (const FAnimNextModuleInstance* ModuleInstance = Context.GetRootGraphInstance().GetModuleInstance())
+		{
+			if (const UAnimNextComponent* AnimNextComponent = Cast<UAnimNextComponent>(ModuleInstance->GetObject()))
+			{
+				const AActor* Actor = AnimNextComponent->GetOwner();
+				check(Actor);
+				if (const USkeletalMeshComponent* SkeletalMeshComponent = Actor->GetComponentByClass<USkeletalMeshComponent>())
+				{
+					InstanceData->ComponentTransform = SkeletalMeshComponent->GetComponentTransform();
+				}
+			}
+		}
+
+#if ENABLE_ANIM_DEBUG 
+		InstanceData->HostObject = Context.GetHostObject();
+#endif // ENABLE_ANIM_DEBUG 
+
+		Context.AppendTask(FAnimNextWarpTestTask::Make(InstanceData));
+	}
+} // namespace UE::AnimNext
+
+
+FAnimNextWarpTestTask FAnimNextWarpTestTask::Make(UE::AnimNext::FWarpTestTrait::FInstanceData* InstanceData)
+{
+	FAnimNextWarpTestTask Task;
+	Task.InstanceData = InstanceData;
+	return Task;
+}
+
+void FAnimNextWarpTestTask::Execute(UE::AnimNext::FEvaluationVM& VM) const
+{
+	using namespace UE::AnimNext;
+
+	const UE::Anim::IAnimRootMotionProvider* RootMotionProvider = UE::Anim::IAnimRootMotionProvider::Get();
+	if (!RootMotionProvider)
+	{
+		UE_LOG(LogAnimNextWarping, Error, TEXT("FAnimNextWarpTestTask::Execute, missing RootMotionProvider"));
+	}
+	else if (const TUniquePtr<FKeyframeState>* Keyframe = VM.PeekValue<TUniquePtr<FKeyframeState>>(KEYFRAME_STACK_NAME, 0))
+	{
+		// snippet on how to extract root motion
+		// FTransform ThisFrameRootMotionTransform = FTransform::Identity;
+		// if (RootMotionProvider->ExtractRootMotion(Keyframe->Get()->Attributes, ThisFrameRootMotionTransform))
+		
+		const FTransform& WarpTo = InstanceData->WarpTo;
+		const FTransform& ComponentTransform = InstanceData->ComponentTransform;
+
+#if ENABLE_ANIM_DEBUG 
+		static const TCHAR* LogName = TEXT("WarpTest");
+		UE_VLOG_SEGMENT_THICK(InstanceData->HostObject, LogName, Display, WarpTo.GetLocation(), WarpTo.GetRotation().GetAxisX() * 100.f + WarpTo.GetLocation(), FColor::Red, 1, TEXT(""));
+		UE_VLOG_SEGMENT_THICK(InstanceData->HostObject, LogName, Display, WarpTo.GetLocation(), WarpTo.GetRotation().GetAxisY() * 100.f + WarpTo.GetLocation(), FColor::Blue, 1, TEXT(""));
+
+		UE_VLOG_SEGMENT_THICK(InstanceData->HostObject, LogName, Display, ComponentTransform.GetLocation(), ComponentTransform.GetRotation().GetAxisX() * 80.f + ComponentTransform.GetLocation(), FColor::Black, 1, TEXT(""));
+		UE_VLOG_SEGMENT_THICK(InstanceData->HostObject, LogName, Display, ComponentTransform.GetLocation(), ComponentTransform.GetRotation().GetAxisY() * 80.f + ComponentTransform.GetLocation(), FColor::Green, 1, TEXT(""));
+#endif // ENABLE_ANIM_DEBUG 
+
+		const FTransform RootMotion = WarpTo.GetRelativeTransform(ComponentTransform);
+
+		RootMotionProvider->OverrideRootMotion(RootMotion, Keyframe->Get()->Attributes);
+	}
+}
diff --git a/Engine/Plugins/Experimental/AnimNextWarping/Source/AnimNextWarping/Private/WarpTestTrait.h b/Engine/Plugins/Experimental/AnimNextWarping/Source/AnimNextWarping/Private/WarpTestTrait.h
new file mode 100644
--- /dev/null
+++ b/Engine/Plugins/Experimental/AnimNextWarping/Source/AnimNextWarping/Private/WarpTestTrait.h
@@ -0,0 +1,75 @@
+// Copyright Epic Games, Inc. All Rights Reserved.
+
+#pragma once
+
+#include "CoreMinimal.h"
+#include "TraitCore/Trait.h"
+#include "TraitInterfaces/IAttributeProvider.h"
+#include "TraitInterfaces/IEvaluate.h"
+#include "TraitInterfaces/IUpdate.h"
+#include "WarpTestTrait.generated.h"
+
+USTRUCT(Experimental, meta = (DisplayName = "WarpTest"))
+struct FWarpTestTraitSharedData : public FAnimNextTraitSharedData
+{
+	GENERATED_BODY()
+
+	// the trait will warp the character looping between Transforms[i] choosing the next one every SecondsToWait
+	UPROPERTY(EditAnywhere, Category = Evaluation, meta = (PinShownByDefault))
+	TArray<FTransform> Transforms;
+
+	// every SecondsToWait we warp to the next Transforms[i]
+	UPROPERTY(EditAnywhere, Category = Evaluation, meta = (PinShownByDefault))
+	float SecondsToWait = 1.f;
+
+	#define TRAIT_LATENT_PROPERTIES_ENUMERATOR(GeneratorMacro) \
+		GeneratorMacro(Transforms) \
+		GeneratorMacro(SecondsToWait) \
+
+	GENERATE_TRAIT_LATENT_PROPERTIES(FWarpTestTraitSharedData, TRAIT_LATENT_PROPERTIES_ENUMERATOR)
+	#undef TRAIT_LATENT_PROPERTIES_ENUMERATOR
+};
+
+namespace UE::AnimNext
+{
+
+struct FWarpTestTrait : FAdditiveTrait, IUpdate, IEvaluate
+{
+	DECLARE_ANIM_TRAIT(FWarpTestTrait, FAdditiveTrait)
+
+	using FSharedData = FWarpTestTraitSharedData;
+
+	struct FInstanceData : FTrait::FInstanceData
+	{
+		int32 CurrentTransformIndex = 0;
+		float CurrentTime = 0.f;
+		FTransform WarpTo = FTransform::Identity;
+		FTransform ComponentTransform = FTransform::Identity;
+#if ENABLE_ANIM_DEBUG 
+		/** Debug Object for VisualLogger */
+		TObjectPtr<const UObject> HostObject = nullptr;
+#endif // ENABLE_ANIM_DEBUG 
+	};
+
+	// IUpdate impl 
+	virtual void PreUpdate(FUpdateTraversalContext& Context, const TTraitBinding<IUpdate>& Binding, const FTraitUpdateState& TraitState) const override;
+
+	// IEvaluate impl
+	virtual void PostEvaluate(FEvaluateTraversalContext& Context, const TTraitBinding<IEvaluate>& Binding) const override;
+};
+
+} // namespace UE::AnimNext
+
+USTRUCT(Experimental)
+struct FAnimNextWarpTestTask : public FAnimNextEvaluationTask
+{
+	GENERATED_BODY()
+
+	DECLARE_ANIM_EVALUATION_TASK(FAnimNextWarpTestTask)
+
+	static FAnimNextWarpTestTask Make(UE::AnimNext::FWarpTestTrait::FInstanceData* InstanceData);
+
+	virtual void Execute(UE::AnimNext::FEvaluationVM& VM) const override;
+
+	UE::AnimNext::FWarpTestTrait::FInstanceData* InstanceData = nullptr;
+};
\ No newline at end of file
