Modifications to AnimNode_Steering and SteeringTrait:
* Allow scaling and additive correction to combine
* Additive correction accounts for root motion and scaling effects
* Additive correction is continuous in the first order (i.e. continuous angular velocity)
* Removed target smoothing since its redundant with additive smoothing
* Deprecate unused 'TargetTime' variable
* Add min / max options for scaling.
* Steering node is marked as EXPERIMENTAL. It was always meant to be.

This fixes these issues:
* Stop excessive steering when changing desired orientation quickly during move starts, pivots etc. This is because scaling could become very large towards end of anim
* Fix 'kink' felt when transitioning from anim with root motion to anim without, as scaling switches to additive

diff --git a/Engine/Plugins/Animation/AnimationWarping/Source/Editor/Public/AnimGraph/AnimGraphNode_Steering.h b/Engine/Plugins/Animation/AnimationWarping/Source/Editor/Public/AnimGraph/AnimGraphNode_Steering.h
--- a/Engine/Plugins/Animation/AnimationWarping/Source/Editor/Public/AnimGraph/AnimGraphNode_Steering.h
+++ b/Engine/Plugins/Animation/AnimationWarping/Source/Editor/Public/AnimGraph/AnimGraphNode_Steering.h
@@ -7,7 +7,7 @@
 
 namespace ENodeTitleType { enum Type : int; }
 
-UCLASS()
+UCLASS(Experimental)
 class ANIMATIONWARPINGEDITOR_API UAnimGraphNode_Steering : public UAnimGraphNode_SkeletalControlBase
 {
 	GENERATED_UCLASS_BODY()
diff --git a/Engine/Plugins/Animation/AnimationWarping/Source/Runtime/Private/BoneControllers/AnimNode_Steering.cpp b/Engine/Plugins/Animation/AnimationWarping/Source/Runtime/Private/BoneControllers/AnimNode_Steering.cpp
--- a/Engine/Plugins/Animation/AnimationWarping/Source/Runtime/Private/BoneControllers/AnimNode_Steering.cpp
+++ b/Engine/Plugins/Animation/AnimationWarping/Source/Runtime/Private/BoneControllers/AnimNode_Steering.cpp
@@ -5,10 +5,10 @@
 #include "Animation/AnimNodeFunctionRef.h"
 #include "Animation/AnimRootMotionProvider.h"
 #include "Animation/AnimSequenceHelpers.h"
+#include "Animation/AnimationMath.h"
 #include "HAL/IConsoleManager.h"
 #include "Animation/AnimTrace.h"
 #include "BoneControllers/AnimNode_OffsetRootBone.h"
-#include "Kismet/KismetMathLibrary.h"
 #include "Logging/LogVerbosity.h"
 #include "VisualLogger/VisualLogger.h"
 
@@ -38,7 +38,7 @@ void FAnimNode_Steering::UpdateInternal(const FAnimationUpdateContext& Context)
 void FAnimNode_Steering::Initialize_AnyThread(const FAnimationInitializeContext& Context)
 {
 	FAnimNode_SkeletalControlBase::Initialize_AnyThread(Context);
-	bResetFilter = true;
+	AngularVelocity = FVector::ZeroVector;
 }
 
 void FAnimNode_Steering::GatherDebugData(FNodeDebugData& DebugData)
@@ -54,28 +54,12 @@ void FAnimNode_Steering::GatherDebugData(FNodeDebugData& DebugData)
 	ComponentPose.GatherDebugData(DebugData);
 }
 
-
 void FAnimNode_Steering::EvaluateSkeletalControl_AnyThread(FComponentSpacePoseContext& Output, TArray<FBoneTransform>& OutBoneTransforms)
 {
 	const float DeltaSeconds = Output.AnimInstanceProxy->GetDeltaSeconds();
 
-	if (DeltaSeconds > 0)
+	if (DeltaSeconds > 0.f)
 	{
-		if (bEnableTargetSmoothing)
-		{
-			if(bResetFilter)
-			{
-				FilteredTarget = TargetOrientation;
-				TargetSmoothingState.Reset();
-				bResetFilter = false;
-			}
-			
-			FilteredTarget = UKismetMathLibrary::QuaternionSpringInterp(FilteredTarget, TargetOrientation, TargetSmoothingState,
-								SmoothTargetStiffness, SmoothTargetDamping, DeltaSeconds, 1, 0, true);
-
-			TargetOrientation = FilteredTarget;
-		}
-		
 		if (Alpha > 0.0f && bAnimNodeSteeringEnabled)
 		{
 			const UE::Anim::IAnimRootMotionProvider* RootMotionProvider = UE::Anim::IAnimRootMotionProvider::Get();
@@ -86,7 +70,8 @@ void FAnimNode_Steering::EvaluateSkeletalControl_AnyThread(FComponentSpacePoseCo
 				FTransform ThisFrameRootMotionTransform = FTransform::Identity;
 				if (RootMotionProvider->ExtractRootMotion(Output.CustomAttributes, ThisFrameRootMotionTransform))
 				{
-					if (ThisFrameRootMotionTransform.GetTranslation().Length() / DeltaSeconds > DisableSteeringBelowSpeed)
+					float CurrentSpeed = ThisFrameRootMotionTransform.GetTranslation().Length() / DeltaSeconds;
+					if (CurrentSpeed > DisableSteeringBelowSpeed)
 					{
 						FQuat RootBoneRotation = RootBoneTransform.GetRotation();
 
@@ -118,7 +103,8 @@ void FAnimNode_Steering::EvaluateSkeletalControl_AnyThread(FComponentSpacePoseCo
 								}
 
 								FQuat PredictedRootMotionQuat = PredictedRootMotionTransform.GetRotation();
-								float PredictedRootMotionYaw = FRotator(PredictedRootMotionQuat).Yaw;
+								FRotator PredictedRootMotionRot = FRotator(PredictedRootMotionQuat);
+								float PredictedRootMotionYaw = PredictedRootMotionRot.Yaw;
 							
 								if (fabs(PredictedRootMotionYaw) > RootMotionThreshold)
 								{
@@ -140,24 +126,45 @@ void FAnimNode_Steering::EvaluateSkeletalControl_AnyThread(FComponentSpacePoseCo
 									}
 
 									float Ratio =  YawToTargetOrientation / PredictedRootMotionYaw ;
+									Ratio = FMath::Clamp(Ratio, MinScaleRatio, MaxScaleRatio);
+									
+									// Account for alpha
+									Ratio = FMath::Lerp(1.0f, Ratio, Alpha);
+
+									FRotator ThisFrameRootMotionRotator(ThisFrameRootMotionTransform.GetRotation());
+									ThisFrameRootMotionRotator.Yaw *= Ratio;
+									ThisFrameRootMotionTransform.SetRotation(FQuat(ThisFrameRootMotionRotator));
 
-									FRotator ThisFrameRootMotionRotation(ThisFrameRootMotionTransform.GetRotation());
-									ThisFrameRootMotionRotation.Yaw *= Ratio;
-								
-									ThisFrameRootMotionTransform.SetRotation(FQuat(ThisFrameRootMotionRotation));
-									RootMotionProvider->OverrideRootMotion(ThisFrameRootMotionTransform, Output.CustomAttributes);
-								
-									return;
+									// Account for future scaling in linear error correction
+									PredictedRootMotionRot.Yaw *= Ratio;
+									PredictedRootMotionQuat = PredictedRootMotionRot.Quaternion();
+
+									DeltaToTargetOrientation = PredictedRootMotionQuat.Inverse() * RootBoneRotation.Inverse() * TargetOrientation;
 								}
 							}
 						}
 						
-						if (ProceduralTargetTime > DeltaSeconds)
-						{	
-							DeltaToTargetOrientation = FQuat::Slerp(FQuat::Identity, DeltaToTargetOrientation,  DeltaSeconds/ProceduralTargetTime);
+						if (CurrentSpeed > DisableAdditiveBelowSpeed)
+						{
+							// Apply linear correction
+							FQuat LinearCorrection = FQuat::Identity;
+							AnimationMath::SimpleSpringDamperQuat(IN OUT LinearCorrection, IN OUT AngularVelocity, DeltaToTargetOrientation, ProceduralTargetTime, DeltaSeconds);
+
+							UE_VLOG_ARROW(Output.AnimInstanceProxy->GetAnimInstanceObject(), "Steering", Display,
+								RootBoneTransform.GetLocation(),
+								RootBoneTransform.GetLocation() + (RootBoneTransform.GetRotation() * LinearCorrection).GetRightVector() * 120,
+								FColor::Magenta, TEXT(""));
+
+							FQuat ThisFrameRotation = ThisFrameRootMotionTransform.GetRotation() * LinearCorrection;
+
+							UE_VLOG_ARROW(Output.AnimInstanceProxy->GetAnimInstanceObject(), "Steering", Display,
+								RootBoneTransform.GetLocation(),
+								RootBoneTransform.GetLocation() + (RootBoneTransform.GetRotation() * ThisFrameRotation).GetRightVector() * 140,
+								FColor::Red, TEXT(""));
+
+							ThisFrameRootMotionTransform.SetRotation(FQuat::Slerp(ThisFrameRootMotionTransform.GetRotation(), ThisFrameRotation, Alpha));
 						}
 
-						ThisFrameRootMotionTransform.SetRotation(FQuat::Slerp(ThisFrameRootMotionTransform.GetRotation(), DeltaToTargetOrientation, Alpha));
 						RootMotionProvider->OverrideRootMotion(ThisFrameRootMotionTransform, Output.CustomAttributes);
 					}
 				}
diff --git a/Engine/Plugins/Animation/AnimationWarping/Source/Runtime/Public/BoneControllers/AnimNode_Steering.h b/Engine/Plugins/Animation/AnimationWarping/Source/Runtime/Public/BoneControllers/AnimNode_Steering.h
--- a/Engine/Plugins/Animation/AnimationWarping/Source/Runtime/Public/BoneControllers/AnimNode_Steering.h
+++ b/Engine/Plugins/Animation/AnimationWarping/Source/Runtime/Public/BoneControllers/AnimNode_Steering.h
@@ -11,50 +11,68 @@ struct FAnimationInitializeContext;
 struct FComponentSpacePoseContext;
 struct FNodeDebugData;
 
-// add procedural delta to the root motion attribute 
+// Add procedural delta to the root motion attribute 
+// This is done via 2 techniques:
+//  1) Scaling the root motion on an animation
+//  2) Adding additional correction to root motion after accounting for the anticipated (potentially scaled) root motion
+// The effects of 1) and 2) combine
 USTRUCT(BlueprintInternalUseOnly)
 struct ANIMATIONWARPINGRUNTIME_API FAnimNode_Steering : public FAnimNode_SkeletalControlBase
 {
 	GENERATED_BODY()
 
+	// Note: We need to explicitly disable warnings on these constructors/operators for clang to be happy with deprecated variables
+	PRAGMA_DISABLE_DEPRECATION_WARNINGS
+	FAnimNode_Steering() = default;
+	~FAnimNode_Steering() = default;
+	FAnimNode_Steering(const FAnimNode_Steering&) = default;
+	FAnimNode_Steering(FAnimNode_Steering&&) = default;
+	FAnimNode_Steering& operator=(const FAnimNode_Steering&) = default;
+	FAnimNode_Steering& operator=(FAnimNode_Steering&&) = default;
+	PRAGMA_ENABLE_DEPRECATION_WARNINGS
+
 	// The Orientation to steer towards
 	UPROPERTY(EditAnywhere, BlueprintReadWrite, Category=Evaluation, meta=(PinShownByDefault))
 	FQuat TargetOrientation = FQuat::Identity;
 
-	UPROPERTY(EditAnywhere, BlueprintReadWrite, Category=Filtering)
-	bool bEnableTargetSmoothing = false;
-
 	// True if input animation is mirrored
 	UPROPERTY(EditAnywhere, BlueprintReadWrite, Category=Filtering, meta = (PinShownByDefault))
 	bool bMirrored = false;
 	
-	UPROPERTY(EditAnywhere, BlueprintReadWrite, Category=Filtering)
-	float SmoothTargetStiffness = 300;
-	
-	UPROPERTY(EditAnywhere, BlueprintReadWrite, Category=Filtering)
-	float SmoothTargetDamping = 1;
-	
-	// The number of seconds in the future before we should reach the TargetOrientation when play animations with no root motion rotation
+	// The number of seconds in the future before we should roughly attempt reach the TargetOrientation via additive correction
 	UPROPERTY(EditAnywhere, BlueprintReadWrite, Category=Evaluation, meta=(PinShownByDefault))
 	float ProceduralTargetTime = 0.2f;
 	
 	// Deprected old/unused parameter, to avoid breaking data
-	UPROPERTY(EditAnywhere, BlueprintReadWrite, Category=Evaluation, meta=(PinShownByDefault))
-	float TargetTime = 0.2f;
+	UE_DEPRECATED(5.6, "Use Procedural target time for the correction time scale and AnimatedTargetTime for the look ahead time on the animation")
+	UPROPERTY(meta=(DeprecatedProperty, DeprecationMessage="Use Procedural target time for the correction time scale and AnimatedTargetTime for the look ahead time on the animation"))
+	float TargetTime_DEPRECATED = 0.2f;
 	
-	// The number of seconds in the future before we should reach the TargetOrientation when playing animations with root motion rotation
+	// The number of seconds in the future to sample the root motion to know how much this animation is expected to turn
 	UPROPERTY(EditAnywhere, BlueprintReadWrite, Category=Evaluation, meta=(PinShownByDefault))
-	float AnimatedTargetTime = 0.2f;
+	float AnimatedTargetTime = 2.0f;
 
-	// If less than this number of degrees of rotation is found over TargetTime seconds in the CurrentAnimAsset, then rotation will be added linearly to reach the TargetOrientation
-	// Otherwise the root motion will be scaled to reach the TargetOrientation over TargetTime seconds
+	// The minimum amount of root motion required to enable root motion scaling.
+	// The root motion is measured from current time to AnimatedTargetTime
 	UPROPERTY(EditAnywhere, DisplayName=RootMotionAngleThreshold, Category=Evaluation)
 	float RootMotionThreshold = 1.0f;
 
-	// below this movement speed (based on the root motion in the animation) disable steering
+	// below this movement speed (based on the root motion in the animation) disable steering completely (both scaling and additive)
 	UPROPERTY(EditAnywhere, Category=Evaluation)
 	float DisableSteeringBelowSpeed = 1.0f;
 
+	// below this movement speed (based on the root motion in the animation) disable steering coming from the additive spring based correction
+	UPROPERTY(EditAnywhere, Category = Evaluation)
+	float DisableAdditiveBelowSpeed = -1.0f;
+
+	// Will clamp the scaling ratio applied to above this threshold. Any error remaining will be compensated for linearly (using ProceduralTargetTime)
+	UPROPERTY(EditAnywhere, Category = Evaluation)
+	float MinScaleRatio = 0.5f;
+
+	// Will clamp the scaling ratio applied to below this threshold. Any error remaining will be compensated for linearly (using ProceduralTargetTime)
+	UPROPERTY(EditAnywhere, Category = Evaluation)
+	float MaxScaleRatio = 1.5f;
+
 	// If bMirrored is set, MirrorDataTable will be used for mirroring the CurrentAnimAsset during prediction
 	UPROPERTY(EditAnywhere, Transient, BlueprintReadWrite, Category = Evaluation, meta = (PinShownByDefault))
 	TObjectPtr<UMirrorDataTable> MirrorDataTable;
@@ -80,10 +98,7 @@ struct ANIMATIONWARPINGRUNTIME_API FAnimNode_Steering : public FAnimNode_Skeleta
 
 private:
 
-	bool bResetFilter = true;
-
-	FQuat FilteredTarget = FQuat::Identity;
-	FQuaternionSpringState TargetSmoothingState;
+	FVector AngularVelocity = FVector::ZeroVector;
 	
 	FTransform RootBoneTransform;
 };
diff --git a/Engine/Plugins/Experimental/AnimNextWarping/Source/AnimNextWarping/Private/SteeringTrait.cpp b/Engine/Plugins/Experimental/AnimNextWarping/Source/AnimNextWarping/Private/SteeringTrait.cpp
--- a/Engine/Plugins/Experimental/AnimNextWarping/Source/AnimNextWarping/Private/SteeringTrait.cpp
+++ b/Engine/Plugins/Experimental/AnimNextWarping/Source/AnimNextWarping/Private/SteeringTrait.cpp
@@ -13,6 +13,7 @@
 #include "TraitInterfaces/IGraphFactory.h"
 #include "TraitInterfaces/ITimeline.h"
 #include "VisualLogger/VisualLogger.h"
+#include "Animation/AnimationMath.h"
 
 bool bAnimNextSteeringTraitEnabled = true;
 static FAutoConsoleVariableRef CVarAnimNextSteeringTraitEnabled(
@@ -87,6 +88,7 @@ namespace UE::AnimNext
 		// Update target orientation, root bone transform, & other properties 
 		InstanceData->TargetOrientation = SharedData->GetTargetOrientation(Binding);
 		InstanceData->RootBoneTransform = SharedData->GetRootBoneTransform(Binding);
+		InstanceData->Alpha = SharedData->GetAlpha(Binding);
 
 #if ENABLE_ANIM_DEBUG 
 		InstanceData->HostObject = Context.GetHostObject();
@@ -113,24 +115,9 @@ void FAnimNextSteeringTask::Execute(UE::AnimNext::FEvaluationVM& VM) const
 	using namespace UE::AnimNext;
 
 	// Execute Steering
-	if (InstanceData->DeltaTime > 0)
+	if (InstanceData->DeltaTime > 0.f)
 	{
-		if (SharedData->bEnableTargetSmoothing)
-		{
-			if (InstanceData->bResetFilter)
-			{
-				InstanceData->FilteredTarget = InstanceData->TargetOrientation;
-				InstanceData->TargetSmoothingState.Reset();
-				InstanceData->bResetFilter = false;
-			}
-
-			InstanceData->FilteredTarget = UKismetMathLibrary::QuaternionSpringInterp(InstanceData->FilteredTarget, InstanceData->TargetOrientation, InstanceData->TargetSmoothingState,
-				SharedData->SmoothTargetStiffness, SharedData->SmoothTargetDamping, InstanceData->DeltaTime, 1, 0, true);
-
-			InstanceData->TargetOrientation = InstanceData->FilteredTarget;
-		}
-
-		if (SharedData->Alpha > 0.0f && bAnimNextSteeringTraitEnabled)
+		if (InstanceData->Alpha > 0.0f && bAnimNextSteeringTraitEnabled)
 		{
 			const UE::Anim::IAnimRootMotionProvider* RootMotionProvider = UE::Anim::IAnimRootMotionProvider::Get();
 			if (!RootMotionProvider)
@@ -144,7 +131,8 @@ void FAnimNextSteeringTask::Execute(UE::AnimNext::FEvaluationVM& VM) const
 				FTransform ThisFrameRootMotionTransform = FTransform::Identity;
 				if (RootMotionProvider->ExtractRootMotion(Keyframe->Get()->Attributes, ThisFrameRootMotionTransform))
 				{
-					if (ThisFrameRootMotionTransform.GetTranslation().Length() / InstanceData->DeltaTime > SharedData->DisableSteeringBelowSpeed)
+					float CurrentSpeed = ThisFrameRootMotionTransform.GetTranslation().Length() / InstanceData->DeltaTime;
+					if (CurrentSpeed > SharedData->DisableSteeringBelowSpeed)
 					{
 						FQuat RootBoneRotation = InstanceData->RootBoneTransform.GetRotation();
 
@@ -168,7 +156,8 @@ void FAnimNextSteeringTask::Execute(UE::AnimNext::FEvaluationVM& VM) const
 							{
 								FTransform PredictedRootMotionTransform = InstanceData->OnExtractRootMotionAttribute.Execute(InstanceData->CurrentAnimAssetTime, SharedData->AnimatedTargetTime, true);
 								FQuat PredictedRootMotionQuat = PredictedRootMotionTransform.GetRotation();
-								float PredictedRootMotionYaw = FRotator(PredictedRootMotionQuat).Yaw;
+								FRotator PredictedRootMotionRot = FRotator(PredictedRootMotionQuat);
+								float PredictedRootMotionYaw = PredictedRootMotionRot.Yaw;
 
 								if (fabs(PredictedRootMotionYaw) > SharedData->RootMotionThreshold)
 								{
@@ -192,24 +181,49 @@ void FAnimNextSteeringTask::Execute(UE::AnimNext::FEvaluationVM& VM) const
 									}
 
 									float Ratio = YawToTargetOrientation / PredictedRootMotionYaw;
+									Ratio = FMath::Clamp(Ratio, SharedData->MinScaleRatio, SharedData->MaxScaleRatio);
 
-									FRotator ThisFrameRootMotionRotation(ThisFrameRootMotionTransform.GetRotation());
-									ThisFrameRootMotionRotation.Yaw *= Ratio;
+									// Account for alpha
+									Ratio = FMath::Lerp(1.0f, Ratio, InstanceData->Alpha);
 
-									ThisFrameRootMotionTransform.SetRotation(FQuat(ThisFrameRootMotionRotation));
-									RootMotionProvider->OverrideRootMotion(ThisFrameRootMotionTransform, Keyframe->Get()->Attributes);
+									FRotator ThisFrameRootMotionRotator(ThisFrameRootMotionTransform.GetRotation());
+									ThisFrameRootMotionRotator.Yaw *= Ratio;
+									ThisFrameRootMotionTransform.SetRotation(FQuat(ThisFrameRootMotionRotator));
 
-									return;
+									// Account for future scaling in linear error correction
+									PredictedRootMotionRot.Yaw *= Ratio;
+									PredictedRootMotionQuat = PredictedRootMotionRot.Quaternion();
+
+									DeltaToTargetOrientation = PredictedRootMotionQuat.Inverse() * RootBoneRotation.Inverse() * InstanceData->TargetOrientation;
 								}
 							}
 						}
 
-						if (SharedData->ProceduralTargetTime > InstanceData->DeltaTime)
+						if (CurrentSpeed > SharedData->DisableAdditiveBelowSpeed)
 						{
-							DeltaToTargetOrientation = FQuat::Slerp(FQuat::Identity, DeltaToTargetOrientation, InstanceData->DeltaTime / SharedData->ProceduralTargetTime);
+							// Apply linear correction
+							FQuat LinearCorrection = FQuat::Identity;
+							AnimationMath::SimpleSpringDamperQuat(LinearCorrection, InstanceData->AngularVelocity, DeltaToTargetOrientation, SharedData->ProceduralTargetTime, InstanceData->DeltaTime);
+
+#if ENABLE_ANIM_DEBUG
+							UE_VLOG_ARROW(InstanceData->HostObject, "Steering", Display,
+								InstanceData->RootBoneTransform.GetLocation(),
+								InstanceData->RootBoneTransform.GetLocation() + (InstanceData->RootBoneTransform.GetRotation() * LinearCorrection).GetRightVector() * 120,
+								FColor::Magenta, TEXT(""));
+#endif
+
+							FQuat ThisFrameRotation = ThisFrameRootMotionTransform.GetRotation() * LinearCorrection;
+
+#if ENABLE_ANIM_DEBUG
+							UE_VLOG_ARROW(InstanceData->HostObject, "Steering", Display,
+								InstanceData->RootBoneTransform.GetLocation(),
+								InstanceData->RootBoneTransform.GetLocation() + (InstanceData->RootBoneTransform.GetRotation() * ThisFrameRotation).GetRightVector() * 140,
+								FColor::Red, TEXT(""));
+#endif
+
+							ThisFrameRootMotionTransform.SetRotation(FQuat::Slerp(ThisFrameRootMotionTransform.GetRotation(), ThisFrameRotation, InstanceData->Alpha));
 						}
 
-						ThisFrameRootMotionTransform.SetRotation(FQuat::Slerp(ThisFrameRootMotionTransform.GetRotation(), DeltaToTargetOrientation, SharedData->Alpha));
 						RootMotionProvider->OverrideRootMotion(ThisFrameRootMotionTransform, Keyframe->Get()->Attributes);
 					}
 				}
diff --git a/Engine/Plugins/Experimental/AnimNextWarping/Source/AnimNextWarping/Private/SteeringTrait.h b/Engine/Plugins/Experimental/AnimNextWarping/Source/AnimNextWarping/Private/SteeringTrait.h
--- a/Engine/Plugins/Experimental/AnimNextWarping/Source/AnimNextWarping/Private/SteeringTrait.h
+++ b/Engine/Plugins/Experimental/AnimNextWarping/Source/AnimNextWarping/Private/SteeringTrait.h
@@ -16,6 +16,12 @@
  * Data needed to execute steering
  * 
  * Some steering data such as current anim asset / playback time is acquired via trait stack interfaces
+ * 
+ * Add procedural delta to the root motion attribute 
+ * This is done via 2 techniques:
+ *  1) Scaling the root motion on an animation
+ *  2) Adding additional correction to root motion after accounting for the anticipated (potentially scaled) root motion
+ * The effects of 1) and 2) combine
  */
 USTRUCT(meta = (DisplayName = "Steering"))
 struct FSteeringTraitSharedData : public FAnimNextTraitSharedData
@@ -36,44 +42,40 @@ struct FSteeringTraitSharedData : public FAnimNextTraitSharedData
 	UPROPERTY(EditAnywhere, Category = Evaluation, meta = (PinShownByDefault))
 	FTransform RootBoneTransform = FTransform::Identity;
 
-	UPROPERTY(EditAnywhere, Category=Filtering)
-	bool bEnableTargetSmoothing = false;
-	
-	UPROPERTY(EditAnywhere, Category=Filtering)
-	float SmoothTargetStiffness = 300;
-	
-	UPROPERTY(EditAnywhere, Category=Filtering)
-	float SmoothTargetDamping = 1;
-	
 	// The number of seconds in the future before we should reach the TargetOrientation when play animations with no root motion rotation
-	UPROPERTY(EditAnywhere, Category=Evaluation, meta=(PinShownByDefault))
+	UPROPERTY(EditAnywhere, Category=Evaluation, meta=(Inline))
 	float ProceduralTargetTime = 0.2f;
 	
 	// The number of seconds in the future before we should reach the TargetOrientation when playing animations with root motion rotation
-	UPROPERTY(EditAnywhere, Category=Evaluation, meta=(PinShownByDefault))
+	UPROPERTY(EditAnywhere, Category=Evaluation, meta=(Inline))
 	float AnimatedTargetTime = 0.2f;
 
-	// If less than this number of degrees of rotation is found over TargetTime seconds in the CurrentAnimAsset, then rotation will be added linearly to reach the TargetOrientation
-	// Otherwise the root motion will be scaled to reach the TargetOrientation over TargetTime seconds
-	UPROPERTY(EditAnywhere, DisplayName=RootMotionAngleThreshold, Category=Evaluation)
+	// The minimum amount of root motion required to enable root motion scaling.
+	// The root motion is measured from current time to AnimatedTargetTime
+	UPROPERTY(EditAnywhere, DisplayName=RootMotionAngleThreshold, Category=Evaluation, meta=(Inline))
 	float RootMotionThreshold = 1.0f;
 
-	// below this movement speed (based on the root motion in the animation) disable steering
-	UPROPERTY(EditAnywhere, Category=Evaluation)
+	// below this movement speed (based on the root motion in the animation) disable steering completely (both scaling and additive)
+	UPROPERTY(EditAnywhere, Category=Evaluation, meta=(Inline))
 	float DisableSteeringBelowSpeed = 1.0f;
 
+	// below this movement speed (based on the root motion in the animation) disable steering coming from the additive spring based correction
+	UPROPERTY(EditAnywhere, Category = Evaluation, Meta=(Inline))
+	float DisableAdditiveBelowSpeed = -1.0f;
+
+	// Will clamp the scaling ratio applied to above this threashold. Any error remaining will be compensated for linearly (using ProceduralTargetTime)
+	UPROPERTY(EditAnywhere, Category = Evaluation, Meta = (Inline))
+	float MinScaleRatio = 0.5f;
+
+	// Will clamp the scaling ratio applied to below this threashold. Any error remaining will be compensated for linearly (using ProceduralTargetTime)
+	UPROPERTY(EditAnywhere, Category = Evaluation, Meta = (Inline))
+	float MaxScaleRatio = 1.5f;
+
 	// Latent pin support boilerplate
 	#define TRAIT_LATENT_PROPERTIES_ENUMERATOR(GeneratorMacro) \
 		GeneratorMacro(Alpha) \
 		GeneratorMacro(TargetOrientation) \
 		GeneratorMacro(RootBoneTransform) \
-		GeneratorMacro(bEnableTargetSmoothing) \
-		GeneratorMacro(SmoothTargetStiffness) \
-		GeneratorMacro(SmoothTargetDamping) \
-		GeneratorMacro(ProceduralTargetTime) \
-		GeneratorMacro(AnimatedTargetTime) \
-		GeneratorMacro(RootMotionThreshold) \
-		GeneratorMacro(DisableSteeringBelowSpeed) \
 
 	GENERATE_TRAIT_LATENT_PROPERTIES(FSteeringTraitSharedData, TRAIT_LATENT_PROPERTIES_ENUMERATOR)
 	#undef TRAIT_LATENT_PROPERTIES_ENUMERATOR
@@ -101,23 +103,20 @@ struct FSteeringTrait : FAdditiveTrait, IUpdate, IEvaluate
 		/** Callback provided by attribute trait on stack to evaluate root motion at a later time */
 		FOnExtractRootMotionAttribute OnExtractRootMotionAttribute = FOnExtractRootMotionAttribute();
 
+	
+		/** Angular velocity of additive correction spring */
+		FVector AngularVelocity = FVector::ZeroVector;
+
+		/** Last root bone transform sampled */
+		FTransform RootBoneTransform = FTransform::Identity;
+
 		/** Current Anim Asset time, used with ProceduralTargetTime to get future root motion, populated during PostEvaluate  */
-		float CurrentAnimAssetTime = 0;
+		float CurrentAnimAssetTime = 0.f;
 
 		/** Delta in seconds between updates, populated during PreUpdate */
 		float DeltaTime = 0.f;
 
-		/** Target after a smoothing spring has been applied */
-		FQuat FilteredTarget = FQuat::Identity;
-		
-		/** Smoothing to apply to Target */
-		FQuaternionSpringState TargetSmoothingState = FQuaternionSpringState();
-
-		/** Last root bone transform sampled */
-		FTransform RootBoneTransform = FTransform::Identity;
-
-		// @TODO: Currently not used, for reinits
-		bool bResetFilter = true;
+		float Alpha = 1.0f;
 
 #if ENABLE_ANIM_DEBUG 
 		/** Debug Object for VisualLogger */
diff --git a/Engine/Source/Runtime/Engine/Internal/Animation/AnimationMath.h b/Engine/Source/Runtime/Engine/Internal/Animation/AnimationMath.h
new file mode 100644
--- /dev/null
+++ b/Engine/Source/Runtime/Engine/Internal/Animation/AnimationMath.h
@@ -0,0 +1,78 @@
+// Copyright Epic Games, Inc. All Rights Reserved.
+
+#pragma once
+
+#include "Math/MathFwd.h"
+#include "Math/Vector.h"
+#include "Math/Quat.h"
+
+namespace AnimationMath
+{
+	// Reference https://theorangeduck.com/page/spring-roll-call 
+
+	inline float SpringHalfLifeToDamping(float HalfLife, float eps = 1e-5f)
+	{
+		return (4.0f * 0.69314718056f) / (HalfLife + eps);
+	}
+
+	inline FVector QuatToScaledAngleAxis(const FQuat& Q)
+	{
+		FQuat QLog = Q.Log();
+		return 2.0f * FVector(QLog.X, QLog.Y, QLog.Z);
+	}
+
+	inline FQuat QuatFromScaledAngleAxis(FVector V)
+	{
+		V /= 2.0f;
+		FQuat QLog(V.X, V.Y, V.Z, 0.0f);
+		return QLog.Exp();
+	}
+
+	// Simplified version of FMath::CriticallyDampedSmoothing where v_goal is assumed to be 0
+	template<typename T>
+	inline void SimpleSpringDamper(
+		T& x,
+		T& v,
+		T x_goal,
+		float halflife,
+		float dt)
+	{
+		float y = SpringHalfLifeToDamping(halflife) / 2.0f;
+		T j0 = x - x_goal;
+		T j1 = v + j0 * y;
+		float eydt = FMath::InvExpApprox(y * dt);
+
+		x = eydt * (j0 + j1 * dt) + x_goal;
+		v = eydt * (v - j1 * y * dt);
+	}
+
+	// Specialized quaternion version, uses FVector for velocity
+	inline void SimpleSpringDamperQuat(
+		FQuat& InOutRotation,
+		FVector& InOutAngularVelocity,
+		const FQuat& TargetRotation,
+		float HalfLife,
+		float DeltaTime)
+	{
+		float y = SpringHalfLifeToDamping(HalfLife) / 2.0f;
+
+		FQuat diff = InOutRotation * TargetRotation.Inverse();
+		diff.EnforceShortestArcWith(FQuat::Identity);
+		FVector j0 = QuatToScaledAngleAxis(diff);
+		FVector j1 = InOutAngularVelocity + j0 * y;
+
+		float eydt = FMath::InvExpApprox(y * DeltaTime);
+
+		InOutRotation = QuatFromScaledAngleAxis(eydt * (j0 + j1 * DeltaTime)) * TargetRotation;
+		InOutAngularVelocity = eydt * (InOutAngularVelocity - j1 * y * DeltaTime);
+	}
+
+	// Specialized quaternion damper, similar to FMath::ExponentialSmoothingApprox but for quats
+	FQuat Damper(const FQuat& CurrentRotation, const FQuat& TargetRotation, float HalfLife, float DeltaTime, float Eps = 1e-5f)
+	{
+		if (HalfLife < Eps)
+			return TargetRotation;
+
+		return FQuat::Slerp(CurrentRotation, TargetRotation, 1.0f - FMath::InvExpApprox((0.69314718056f * DeltaTime) / (HalfLife)));
+	}
+}
\ No newline at end of file
