Implement strafe warping with embedded simple 2 bone IK in anim next (Experimental)
* Strafe warping is a port of "Orientation Warping" from animation blueprint
* However instead of writing to IK effectors in component space, the trait directly applies a 2 bone IK solve.

AnimNext dependent changes:
* Add a way to convert from local to component and back again for rotation and translation
* Rename some methods for clarity in TransformArrayOperations
* Add a hashmap lookup for resolving bone indices from bone names

diff --git a/Engine/Plugins/Experimental/AnimNext/Source/AnimNext/Private/GenerationTools.cpp b/Engine/Plugins/Experimental/AnimNext/Source/AnimNext/Private/GenerationTools.cpp
--- a/Engine/Plugins/Experimental/AnimNext/Source/AnimNext/Private/GenerationTools.cpp
+++ b/Engine/Plugins/Experimental/AnimNext/Source/AnimNext/Private/GenerationTools.cpp
@@ -136,9 +136,13 @@ struct FCompareBoneIndexType
 				TArray<FBoneIndexType> SkeletonBoneIndexToLODBoneIndexMap;
 				TArray<FBoneIndexType> MeshBoneIndexToLODBoneIndexMap;
 				TArray<FBoneIndexType> LODBoneIndexToSkeletonBoneIndexMap;
+				TMap<FName, FBoneIndexType> BoneNameToLODIndexMap;
+				TArray<FBoneIndexType> MeshBoneIndexToLODBoneIndex;
 				SkeletonBoneIndexToLODBoneIndexMap.Init(INDEX_NONE, NumSkelBones);
 				MeshBoneIndexToLODBoneIndexMap.Init(INDEX_NONE, NumMeshBones);
 				LODBoneIndexToSkeletonBoneIndexMap.Init(INDEX_NONE, NumOrderedBones);
+				BoneNameToLODIndexMap.Reserve(NumOrderedBones);
+				MeshBoneIndexToLODBoneIndex.Init(INDEX_NONE, NumMeshBones);
 
 				for (int32 LODBoneIndex = 0; LODBoneIndex < NumOrderedBones; ++LODBoneIndex)
 				{
@@ -166,15 +170,14 @@ struct FCompareBoneIndexType
 					{
 						MeshBoneIndexToLODBoneIndexMap[static_cast<FBoneIndexType>(MeshBoneIndex.GetInt())] = LODBoneIndex;
 					}
-				}
 
-				TArray<FBoneIndexType> MeshBoneIndexToLODBoneIndex;
-				MeshBoneIndexToLODBoneIndex.Init(INDEX_NONE, NumMeshBones);
+					FName BoneName = SkeletalMesh->GetRefSkeleton().GetBoneName(MeshBoneIndex.GetInt());
+					if (BoneName.IsValid())
+					{
+						BoneNameToLODIndexMap.Add(BoneName, LODBoneIndex);
+					}
 
-				for (int32 LODBoneIndex = 0; LODBoneIndex < NumOrderedBones; ++LODBoneIndex)
-				{
-					const int32 MeshBoneIndex = LODBoneIndexToMeshBoneIndexMap[LODBoneIndex];
-					MeshBoneIndexToLODBoneIndex[MeshBoneIndex] = LODBoneIndex;
+					MeshBoneIndexToLODBoneIndex[MeshBoneIndex.GetInt()] = LODBoneIndex;
 				}
 
 				const FReferenceSkeleton& RefSkeleton = SkeletalMesh->GetRefSkeleton();
@@ -199,6 +202,7 @@ struct FCompareBoneIndexType
 					{ SkeletonBoneIndexToLODBoneIndexMap },
 					{ MeshBoneIndexToLODBoneIndexMap },
 					LODNumBones,
+					BoneNameToLODIndexMap,
 					bCanGenerateSingleBonesList);
 
 				ReferencePoseGenerated = true;
diff --git a/Engine/Plugins/Experimental/AnimNext/Source/AnimNext/Private/TransformArrayOperations_AoS.cpp b/Engine/Plugins/Experimental/AnimNext/Source/AnimNext/Private/TransformArrayOperations_AoS.cpp
--- a/Engine/Plugins/Experimental/AnimNext/Source/AnimNext/Private/TransformArrayOperations_AoS.cpp
+++ b/Engine/Plugins/Experimental/AnimNext/Source/AnimNext/Private/TransformArrayOperations_AoS.cpp
@@ -51,7 +51,7 @@ namespace UE::AnimNext
 		}
 	}
 
-	void ConvertPoseLocalToMesh(const FTransformArrayAoSView& Input, const TArrayView<const FBoneIndexType>& LODBoneIndexToParentLODBoneIndexMap)
+	void ConvertPoseLocalToMeshRotation(const FTransformArrayAoSView& Input, const TArrayView<const FBoneIndexType>& LODBoneIndexToParentLODBoneIndexMap)
 	{
 		// @TODO: This method is untested as AoS doesn't seem to compile? 
 		const int32 NumLODBoneIndexes = LODBoneIndexToParentLODBoneIndexMap.Num();
@@ -64,7 +64,7 @@ namespace UE::AnimNext
 		}
 	}
 
-	void ConvertPoseMeshToLocal(const FTransformArrayAoSView& Input, const TArrayView<const FBoneIndexType>& LODBoneIndexToParentLODBoneIndexMap)
+	void ConvertPoseMeshToLocalRotation(const FTransformArrayAoSView& Input, const TArrayView<const FBoneIndexType>& LODBoneIndexToParentLODBoneIndexMap)
 	{
 		// @TODO: This method is untested as AoS doesn't seem to compile? 
 		const int32 NumLODBoneIndexes = LODBoneIndexToParentLODBoneIndexMap.Num();
@@ -77,6 +77,37 @@ namespace UE::AnimNext
 		}
 	}
 
+	void ConvertPoseLocalToMeshRotationTranslation(const FTransformArrayAoSView& Input, const TArrayView<const FBoneIndexType>& LODBoneIndexToParentLODBoneIndexMap)
+	{
+		// @TODO: This method is untested as AoS doesn't seem to compile? 
+		const int32 NumLODBoneIndexes = LODBoneIndexToParentLODBoneIndexMap.Num();
+
+		for (int32 LODBoneIndex = 1; LODBoneIndex < NumLODBoneIndexes; ++LODBoneIndex)
+		{
+			const FBoneIndexType ParentLODBoneIndex = LODBoneIndexToParentLODBoneIndexMap[LODBoneIndex];
+			const FQuat MeshSpaceRotation = Input[ParentLODBoneIndex].GetRotation() * Input[LODBoneIndex].GetRotation();
+			const FVector MeshSpaceTranslation = Input[ParentLODBoneIndex].GetTranslation() + Input[ParentLODBoneIndex].GetRotation().RotateVector(Input[LODBoneIndex].GetTranslation());
+			Input[LODBoneIndex].SetRotation(MeshSpaceRotation);
+			Input[LODBoneIndex].SetTranslation(MeshSpaceTranslation);
+		}
+	}
+
+	void ConvertPoseMeshToLocalRotationTranslation(const FTransformArrayAoSView& Input, const TArrayView<const FBoneIndexType>& LODBoneIndexToParentLODBoneIndexMap)
+	{
+		// @TODO: This method is untested as AoS doesn't seem to compile? 
+		const int32 NumLODBoneIndexes = LODBoneIndexToParentLODBoneIndexMap.Num();
+
+		for (int32 LODBoneIndex = NumLODBoneIndexes - 1; LODBoneIndex > 0; --LODBoneIndex)
+		{
+			const FBoneIndexType ParentLODBoneIndex = LODBoneIndexToParentLODBoneIndexMap[LODBoneIndex];
+			const FQuat LocalSpaceRotation = Input[ParentLODBoneIndex].GetRotation().Inverse() * Input[LODBoneIndex].GetRotation();
+			FVector LocalSpaceTranslation = Input[LODBoneIndex].GetTranslation() - Input[ParentLODBoneIndex].GetTranslation();
+			LocalSpaceTranslation = Input[ParentLODBoneIndex].GetRotation().UnrotateVector(LocalSpaceTranslation);
+			Input[LODBoneIndex].SetRotation(LocalSpaceRotation);
+			Input[LODBoneIndex].SetTranslation(LocalSpaceTranslation);
+		}
+	}
+
 	void BlendWithIdentityAndAccumulate(const FTransformArrayAoSView& Base, const FTransformArrayAoSConstView& Additive, const float BlendWeight)
 	{
 		const ScalarRegister VBlendWeight(BlendWeight);
@@ -95,11 +126,11 @@ namespace UE::AnimNext
 		const TArrayView<const FBoneIndexType>& LODBoneIndexToParentLODBoneIndexMap, const float BlendWeight)
 	{
 		// @TODO: This method is untested as AoS doesn't seem to compile? 
-		ConvertPoseLocalToMesh(Base, LODBoneIndexToParentLODBoneIndexMap);
+		ConvertPoseLocalToMeshRotation(Base, LODBoneIndexToParentLODBoneIndexMap);
 
 		BlendWithIdentityAndAccumulate(Base, Additive, BlendWeight);
 
-		ConvertPoseMeshToLocal(Base, LODBoneIndexToParentLODBoneIndexMap);
+		ConvertPoseMeshToLocalRotation(Base, LODBoneIndexToParentLODBoneIndexMap);
 	}
 
 	void BlendOverwriteWithScale(const FTransformArrayAoSView& Dest, const FTransformArrayAoSConstView& Source, const float ScaleWeight)
diff --git a/Engine/Plugins/Experimental/AnimNext/Source/AnimNext/Private/TransformArrayOperations_SoA.cpp b/Engine/Plugins/Experimental/AnimNext/Source/AnimNext/Private/TransformArrayOperations_SoA.cpp
--- a/Engine/Plugins/Experimental/AnimNext/Source/AnimNext/Private/TransformArrayOperations_SoA.cpp
+++ b/Engine/Plugins/Experimental/AnimNext/Source/AnimNext/Private/TransformArrayOperations_SoA.cpp
@@ -100,7 +100,7 @@ namespace UE::AnimNext
 		}
 	}
 
-	void ConvertPoseLocalToMesh(const FTransformArraySoAView& Input, const TArrayView<const FBoneIndexType>& LODBoneIndexToParentLODBoneIndexMap)
+	void ConvertPoseLocalToMeshRotation(const FTransformArraySoAView& Input, const TArrayView<const FBoneIndexType>& LODBoneIndexToParentLODBoneIndexMap)
 	{
 		const int32 NumLODBoneIndexes = LODBoneIndexToParentLODBoneIndexMap.Num();
 
@@ -112,7 +112,7 @@ namespace UE::AnimNext
 		}
 	}
 
-	void ConvertPoseMeshToLocal(const FTransformArraySoAView& Input, const TArrayView<const FBoneIndexType>& LODBoneIndexToParentLODBoneIndexMap)
+	void ConvertPoseMeshToLocalRotation(const FTransformArraySoAView& Input, const TArrayView<const FBoneIndexType>& LODBoneIndexToParentLODBoneIndexMap)
 	{
 		const int32 NumLODBoneIndexes = LODBoneIndexToParentLODBoneIndexMap.Num();
 
@@ -124,6 +124,34 @@ namespace UE::AnimNext
 		}
 	}
 
+	void ConvertPoseLocalToMeshRotationTranslation(const FTransformArraySoAView& Input, const TArrayView<const FBoneIndexType>& LODBoneIndexToParentLODBoneIndexMap)
+	{
+		const int32 NumLODBoneIndexes = LODBoneIndexToParentLODBoneIndexMap.Num();
+
+		for (int32 LODBoneIndex = 1; LODBoneIndex < NumLODBoneIndexes; ++LODBoneIndex)
+		{
+			const FBoneIndexType ParentLODBoneIndex = LODBoneIndexToParentLODBoneIndexMap[LODBoneIndex];
+			const FQuat MeshSpaceRotation = Input[ParentLODBoneIndex].GetRotation() * Input[LODBoneIndex].GetRotation();
+			const FVector MeshSpaceTranslation = Input[ParentLODBoneIndex].GetTranslation() + Input[ParentLODBoneIndex].GetRotation().RotateVector(Input[LODBoneIndex].GetTranslation());
+			Input.Rotations[LODBoneIndex] = MeshSpaceRotation;
+			Input.Translations[LODBoneIndex] = MeshSpaceTranslation;
+		}
+	}
+
+	void ConvertPoseMeshToLocalRotationTranslation(const FTransformArraySoAView& Input, const TArrayView<const FBoneIndexType>& LODBoneIndexToParentLODBoneIndexMap)
+	{
+		const int32 NumLODBoneIndexes = LODBoneIndexToParentLODBoneIndexMap.Num();
+
+		for (int32 LODBoneIndex = NumLODBoneIndexes - 1; LODBoneIndex > 0; --LODBoneIndex)
+		{
+			const FBoneIndexType ParentLODBoneIndex = LODBoneIndexToParentLODBoneIndexMap[LODBoneIndex];
+			const FQuat LocalSpaceRotation = Input[ParentLODBoneIndex].GetRotation().Inverse() * Input[LODBoneIndex].GetRotation();
+			Input.Rotations[LODBoneIndex] = LocalSpaceRotation;
+			Input.Translations[LODBoneIndex] -= Input[ParentLODBoneIndex].GetTranslation();
+			Input.Translations[LODBoneIndex] = Input[ParentLODBoneIndex].GetRotation().UnrotateVector(Input[LODBoneIndex].GetTranslation());
+		}
+	}
+
 	void BlendWithIdentityAndAccumulate(const FTransformArraySoAView& Base, const FTransformArraySoAConstView& Additive, const float BlendWeight)
 	{
 		const int32 NumTransforms = Base.Num();
@@ -178,11 +206,12 @@ namespace UE::AnimNext
 		const FTransformArraySoAView& Base, const FTransformArraySoAConstView& Additive, 
 		const TArrayView<const FBoneIndexType>& LODBoneIndexToParentLODBoneIndexMap, const float BlendWeight)
 	{
-		ConvertPoseLocalToMesh(Base, LODBoneIndexToParentLODBoneIndexMap);
+		// The additive blend is done in component space for rotations, but local for translation and scale
+		ConvertPoseLocalToMeshRotation(Base, LODBoneIndexToParentLODBoneIndexMap);
 
 		BlendWithIdentityAndAccumulate(Base, Additive, BlendWeight);
 
-		ConvertPoseMeshToLocal(Base, LODBoneIndexToParentLODBoneIndexMap);
+		ConvertPoseMeshToLocalRotation(Base, LODBoneIndexToParentLODBoneIndexMap);
 	}
 
 	void BlendOverwriteWithScale(const FTransformArraySoAView& Dest, const FTransformArraySoAConstView& Source, const float ScaleWeight)
diff --git a/Engine/Plugins/Experimental/AnimNext/Source/AnimNext/Public/LODPose.h b/Engine/Plugins/Experimental/AnimNext/Source/AnimNext/Public/LODPose.h
--- a/Engine/Plugins/Experimental/AnimNext/Source/AnimNext/Public/LODPose.h
+++ b/Engine/Plugins/Experimental/AnimNext/Source/AnimNext/Public/LODPose.h
@@ -203,6 +203,62 @@ struct FLODPose
 		}
 	}
 
+	// Query to find a LODBoneIndex for an associated BoneName
+	// Returns INDEX_NONE if missing ReferencePose or if no bone found for a given name
+	FBoneIndexType FindLODBoneIndexFromBoneName(FName BoneName) const
+	{
+		if (RefPose != nullptr)
+		{
+			return RefPose->FindLODBoneIndexFromBoneName(BoneName);
+		}
+		return INDEX_NONE;
+	}
+
+	// Query whether bone with LODIndex ChildLODBoneIndex is a child of bone with LODIndex ParentLODBoneIndex
+	// Returns:
+	//	True - ChildLODBoneIndex is a child of ParentLODBoneIndex
+	//  False - ChildLODBoneIndex is not a child of ParentLODBoneIndex
+	bool IsBoneChildOf(FBoneIndexType ChildLODBoneIndex, FBoneIndexType ParentLODBoneIndex) const
+	{
+		check(ChildLODBoneIndex != INDEX_NONE);
+		check(ParentLODBoneIndex != INDEX_NONE);
+		check(ChildLODBoneIndex != ParentLODBoneIndex);
+		
+		if (RefPose == nullptr)
+		{
+			return false;
+		}
+
+		const TArrayView<const FBoneIndexType> LODBoneIndexToParentMap = RefPose->GetLODBoneIndexToParentLODBoneIndexMap(LODLevel); 
+		
+		FBoneIndexType ParentOfChild = LODBoneIndexToParentMap[ChildLODBoneIndex];
+		while (ParentOfChild != INDEX_NONE)
+		{
+			if (ParentOfChild == ParentLODBoneIndex)
+			{
+				return true;
+			}
+
+			// Keep looking
+			ChildLODBoneIndex = ParentOfChild;
+			ParentOfChild = LODBoneIndexToParentMap[ChildLODBoneIndex];
+		}
+
+		return false;
+	}
+
+	// Get the LODBoneIndex of the parent of the bone at ChildLODBoneIndex
+	// Returns INDEX_NONE if no reference pose or if ChilddLODBoneIndex is the root
+	FBoneIndexType GetLODBoneParentIndex(FBoneIndexType ChildLODBoneIndex) const
+	{
+		if (RefPose == nullptr)
+		{
+			return INDEX_NONE;
+		}
+		
+		return RefPose->GetLODParentBoneIndex(LODLevel, ChildLODBoneIndex);
+	}
+
 	const USkeleton* GetSkeletonAsset() const
 	{
 		return RefPose != nullptr ? RefPose->Skeleton.Get() : nullptr;
diff --git a/Engine/Plugins/Experimental/AnimNext/Source/AnimNext/Public/ReferencePose.h b/Engine/Plugins/Experimental/AnimNext/Source/AnimNext/Public/ReferencePose.h
--- a/Engine/Plugins/Experimental/AnimNext/Source/AnimNext/Public/ReferencePose.h
+++ b/Engine/Plugins/Experimental/AnimNext/Source/AnimNext/Public/ReferencePose.h
@@ -27,7 +27,7 @@ enum class EReferencePoseGenerationFlags : uint8
 
 ENUM_CLASS_FLAGS(EReferencePoseGenerationFlags);
 
-template <typename AllocatorType>
+template <typename AllocatorType, typename SetAllocator>
 struct TReferencePose
 {
 	// Transform array of our bind pose sorted by LOD, allows us to truncate the array for a specific LOD
@@ -67,6 +67,9 @@ struct TReferencePose
 	// Mapping of mesh bone indices to mesh parent indices for each bone
 	TArray<FBoneIndexType, AllocatorType> MeshBoneIndexToParentMeshBoneIndexMap;
 
+	// Mapping of bone names to LOD Bone Indices
+	TMap<FName, FBoneIndexType, SetAllocator> BoneNameToLODBoneIndexMap;
+
 	TWeakObjectPtr<const USkeletalMeshComponent> SkeletalMeshComponent = nullptr;
 	TWeakObjectPtr<const USkeletalMesh> SkeletalMesh = nullptr;
 	TWeakObjectPtr<const USkeleton> Skeleton = nullptr;
@@ -105,7 +108,8 @@ struct TReferencePose
 		, const TArray<FBoneIndexType>& InSkeletonBoneIndexToLODBoneIndexMap
 		, const TArray<FBoneIndexType>& InMeshBoneIndexToLODBoneIndexMap
 		, const TArray<int32, AllocatorType>& InLODNumBones
-		, bool bFastPath = false)
+		, const TMap<FName, FBoneIndexType>& InNameToLODBoneIndexMap
+		, bool bFastPath)
 	{
 		const int32 NumBonesLOD0 = !InLODNumBones.IsEmpty() ? InLODNumBones[0] : 0;
 		const int32 NumBonesMesh = RefSkeleton.GetRefBoneInfo().Num();
@@ -117,6 +121,7 @@ struct TReferencePose
 		LODBoneIndexToSkeletonBoneIndexMapPerLOD = InLODBoneIndexToSkeletonBoneIndexMapPerLOD;
 		SkeletonBoneIndexToLODBoneIndexMap = InSkeletonBoneIndexToLODBoneIndexMap;
 		MeshBoneIndexToLODBoneIndexMap = InMeshBoneIndexToLODBoneIndexMap;
+		BoneNameToLODBoneIndexMap = InNameToLODBoneIndexMap;
 		LODNumBones = InLODNumBones;
 
 		const TArray<FTransform>& RefBonePose = RefSkeleton.GetRefBonePose();
@@ -197,6 +202,22 @@ struct TReferencePose
 		return MeshBoneIndexToLODBoneIndexMap;
 	}
 
+	const TMap<FName, FBoneIndexType>& GetBoneNameToLODBoneIndexMap() const
+	{
+		return BoneNameToLODBoneIndexMap;
+	}
+
+	// Query to find a LODBoneIndex for an associated BoneName
+	// Returns INDEX_NONE if no bone found for a given name
+	FBoneIndexType FindLODBoneIndexFromBoneName(FName BoneName) const
+	{
+		if (const FBoneIndexType* BoneIndex = BoneNameToLODBoneIndexMap.Find(BoneName))
+		{
+			return *BoneIndex;
+		}
+		return INDEX_NONE;
+	}
+
 	int32 GetMeshBoneIndexFromLODBoneIndex(int32 LODBoneIndex) const
 	{
 		const int32 NumBonesLOD0 = LODBoneIndexToMeshBoneIndexMapPerLOD[0].Num();
@@ -290,7 +311,7 @@ struct TReferencePose
 /**
 * Hide the allocator usage
 */
-using FReferencePose = TReferencePose<FDefaultAllocator>;
+using FReferencePose = TReferencePose<FDefaultAllocator, FDefaultSetAllocator>;
 
 } // namespace UE::AnimNext
 
diff --git a/Engine/Plugins/Experimental/AnimNext/Source/AnimNext/Public/TransformArray.h b/Engine/Plugins/Experimental/AnimNext/Source/AnimNext/Public/TransformArray.h
--- a/Engine/Plugins/Experimental/AnimNext/Source/AnimNext/Public/TransformArray.h
+++ b/Engine/Plugins/Experimental/AnimNext/Source/AnimNext/Public/TransformArray.h
@@ -378,6 +378,16 @@ struct TTransformArraySoA
 		std::copy(Other.Scales3D.begin() + Index, Other.Scales3D.begin() + (Index + NumTransforms), Scales3D.begin());
 	}
 
+	inline void CopyTransforms(const FTransformArraySoAConstView& Other, int32 Index, int32 NumTransforms)
+	{
+		check(Index < Num() && Index < Other.Num());
+		check(Index + (NumTransforms - 1) < Num() && Index + (NumTransforms - 1) < Other.Num());
+
+		std::copy(Other.Translations.begin() + Index, Other.Translations.begin() + (Index + NumTransforms), Translations.begin());
+		std::copy(Other.Rotations.begin() + Index, Other.Rotations.begin() + (Index + NumTransforms), Rotations.begin());
+		std::copy(Other.Scales3D.begin() + Index, Other.Scales3D.begin() + (Index + NumTransforms), Scales3D.begin());
+	}
+
 	/** Set this transform array to the weighted blend of the supplied two transforms. */
 	void Blend(const TTransformArraySoA& AtomArray1, const TTransformArraySoA& AtomArray2, float BlendWeight)
 	{
diff --git a/Engine/Plugins/Experimental/AnimNext/Source/AnimNext/Public/TransformArrayOperations.h b/Engine/Plugins/Experimental/AnimNext/Source/AnimNext/Public/TransformArrayOperations.h
--- a/Engine/Plugins/Experimental/AnimNext/Source/AnimNext/Public/TransformArrayOperations.h
+++ b/Engine/Plugins/Experimental/AnimNext/Source/AnimNext/Public/TransformArrayOperations.h
@@ -23,13 +23,21 @@ namespace UE::AnimNext
 	ANIMNEXT_API void NormalizeRotations(const FTransformArrayAoSView& Input);
 	ANIMNEXT_API void NormalizeRotations(const FTransformArraySoAView& Input);
 
-	// Convert a pose from local space to mesh space
-	ANIMNEXT_API void ConvertPoseLocalToMesh(const FTransformArrayAoSView& Input, const TArrayView<const FBoneIndexType>& LODBoneIndexToParentLODBoneIndexMap);
-	ANIMNEXT_API void ConvertPoseLocalToMesh(const FTransformArraySoAView& Input, const TArrayView<const FBoneIndexType>& LODBoneIndexToParentLODBoneIndexMap);
+	// Convert a pose from local space to mesh space (rotations only)
+	ANIMNEXT_API void ConvertPoseLocalToMeshRotation(const FTransformArrayAoSView& Input, const TArrayView<const FBoneIndexType>& LODBoneIndexToParentLODBoneIndexMap);
+	ANIMNEXT_API void ConvertPoseLocalToMeshRotation(const FTransformArraySoAView& Input, const TArrayView<const FBoneIndexType>& LODBoneIndexToParentLODBoneIndexMap);
 
-	// Convert a pose from mesh space to local space
-	ANIMNEXT_API void ConvertPoseMeshToLocal(const FTransformArrayAoSView& Input, const TArrayView<const FBoneIndexType>& LODBoneIndexToParentLODBoneIndexMap);
-	ANIMNEXT_API void ConvertPoseMeshToLocal(const FTransformArraySoAView& Input, const TArrayView<const FBoneIndexType>& LODBoneIndexToParentLODBoneIndexMap);
+	// Convert a pose from mesh space to local space (rotations only)
+	ANIMNEXT_API void ConvertPoseMeshToLocalRotation(const FTransformArrayAoSView& Input, const TArrayView<const FBoneIndexType>& LODBoneIndexToParentLODBoneIndexMap);
+	ANIMNEXT_API void ConvertPoseMeshToLocalRotation(const FTransformArraySoAView& Input, const TArrayView<const FBoneIndexType>& LODBoneIndexToParentLODBoneIndexMap);
+	
+	// Convert a pose from local space to mesh space (rotation and translation only)
+	ANIMNEXT_API void ConvertPoseLocalToMeshRotationTranslation(const FTransformArrayAoSView& Input, const TArrayView<const FBoneIndexType>& LODBoneIndexToParentLODBoneIndexMap);
+	ANIMNEXT_API void ConvertPoseLocalToMeshRotationTranslation(const FTransformArraySoAView& Input, const TArrayView<const FBoneIndexType>& LODBoneIndexToParentLODBoneIndexMap);
+
+	// Convert a pose from mesh space to local space (rotation and translation only)
+	ANIMNEXT_API void ConvertPoseMeshToLocalRotationTranslation(const FTransformArrayAoSView& Input, const TArrayView<const FBoneIndexType>& LODBoneIndexToParentLODBoneIndexMap);
+	ANIMNEXT_API void ConvertPoseMeshToLocalRotationTranslation(const FTransformArraySoAView& Input, const TArrayView<const FBoneIndexType>& LODBoneIndexToParentLODBoneIndexMap);
 
 	// The additive transform view is blended with the additive identity using the provided blend weight
 	// We then accumulate the resulting transforms on top of the base transforms
diff --git a/Engine/Plugins/Experimental/AnimNextWarping/Source/AnimNextWarping/AnimNextWarping.Build.cs b/Engine/Plugins/Experimental/AnimNextWarping/Source/AnimNextWarping/AnimNextWarping.Build.cs
--- a/Engine/Plugins/Experimental/AnimNextWarping/Source/AnimNextWarping/AnimNextWarping.Build.cs
+++ b/Engine/Plugins/Experimental/AnimNextWarping/Source/AnimNextWarping/AnimNextWarping.Build.cs
@@ -11,6 +11,7 @@ namespace UnrealBuildTool.Rules
 				{
 					"Core",
 					"CoreUObject",
+					"AnimationCore",
 					"RigVM",
 					"Engine",
 					"AnimNext",
diff --git a/Engine/Plugins/Experimental/AnimNextWarping/Source/AnimNextWarping/Private/Module.cpp b/Engine/Plugins/Experimental/AnimNextWarping/Source/AnimNextWarping/Private/Module.cpp
--- a/Engine/Plugins/Experimental/AnimNextWarping/Source/AnimNextWarping/Private/Module.cpp
+++ b/Engine/Plugins/Experimental/AnimNextWarping/Source/AnimNextWarping/Private/Module.cpp
@@ -6,6 +6,7 @@
 #include "Modules/ModuleManager.h"
 #include "RigVMCore/RigVMRegistry.h"
 #include "Animation/AnimSequence.h"
+#include "StrafeWarpingTrait.h"
 
 #include "AnimNextWarpingLog.h"
 #include "EvaluationVM/EvaluationVM.h"
@@ -20,6 +21,14 @@ class FModule : public IModuleInterface
 public:
 	virtual void StartupModule() override
 	{
+		static UScriptStruct* const AllowedStructTypes[] =
+		{
+			FStrafeWarpFootData::StaticStruct(),
+		};
+
+		FRigVMRegistry& RigVMRegistry = FRigVMRegistry::Get();
+		RigVMRegistry.RegisterStructTypes(AllowedStructTypes);
+		
 		// @TODO: Register any pin types needed for warping here for RigVM, currently we have none
 		//static TPair<UClass*, FRigVMRegistry::ERegisterObjectOperation> const AllowedObjectTypes[] =
 		//{
diff --git a/Engine/Plugins/Experimental/AnimNextWarping/Source/AnimNextWarping/Private/StrafeWarpingTrait.cpp b/Engine/Plugins/Experimental/AnimNextWarping/Source/AnimNextWarping/Private/StrafeWarpingTrait.cpp
new file mode 100644
--- /dev/null
+++ b/Engine/Plugins/Experimental/AnimNextWarping/Source/AnimNextWarping/Private/StrafeWarpingTrait.cpp
@@ -0,0 +1,589 @@
+// Copyright Epic Games, Inc. All Rights Reserved.
+
+#include "StrafeWarpingTrait.h"
+
+#include "AnimNextWarpingLog.h"
+#include "TwoBoneIK.h"
+#include "Animation/AnimRootMotionProvider.h"
+#include "Animation/AnimSequenceBase.h"
+#include "EvaluationVM/EvaluationVM.h"
+#include "EvaluationVM/KeyframeState.h"
+#include "TraitInterfaces/IAttributeProvider.h"
+#include "TraitInterfaces/IGraphFactory.h"
+#include "TraitInterfaces/ITimeline.h"
+#include "VisualLogger/VisualLogger.h"
+
+#if ENABLE_ANIM_DEBUG
+bool bAnimNextStrafeWarpingTraitEnabled = true;
+static FAutoConsoleVariableRef CVarAnimNextStrafeWarpingTraitEnabled(
+	TEXT("a.AnimNext.StrafeWarpingTrait.Enabled"),
+	bAnimNextStrafeWarpingTraitEnabled,
+	TEXT("True will enable strafe warping for AnimNext. Equivalent to setting alpha to non-zero.")
+);
+#else
+constexpr bool bAnimNextStrafeWarpingTraitEnabled = true;
+#endif
+
+////////////////////////////////////////////////////////////////////////////////////////////////////////////////////////
+// FStrafeWarpingTrait
+
+namespace UE::AnimNext
+{
+	static FVector GetAxisVector(const EAxis::Type& InAxis)
+	{
+		switch (InAxis)
+		{
+		case EAxis::X:
+			return FVector::ForwardVector;
+		case EAxis::Y:
+			return FVector::RightVector;
+		default:
+			return FVector::UpVector;
+		};
+	}
+
+	static float SignedAngleRadBetweenNormals(const FVector& From, const FVector& To, const FVector& Axis)
+	{
+		const float FromDotTo = FVector::DotProduct(From, To);
+		const float Angle = FMath::Acos(FromDotTo);
+		const FVector Cross = FVector::CrossProduct(From, To);
+		const float Dot = FVector::DotProduct(Cross, Axis);
+		return Dot >= 0 ? Angle : -Angle;
+	}
+	
+	AUTO_REGISTER_ANIM_TRAIT(FStrafeWarpingTrait)
+
+	#define TRAIT_INTERFACE_ENUMERATOR(GeneratorMacro) \
+		GeneratorMacro(IUpdate) \
+		GeneratorMacro(IEvaluate) \
+
+	// Trait implementation boilerplate
+	GENERATE_ANIM_TRAIT_IMPLEMENTATION(FStrafeWarpingTrait, TRAIT_INTERFACE_ENUMERATOR, NULL_ANIM_TRAIT_INTERFACE_ENUMERATOR, NULL_ANIM_TRAIT_EVENT_ENUMERATOR)
+	#undef TRAIT_INTERFACE_ENUMERATOR
+
+	void FStrafeWarpingTrait::PreUpdate(FUpdateTraversalContext& Context, const TTraitBinding<IUpdate>& Binding, const FTraitUpdateState& TraitState) const
+	{
+		FInstanceData* InstanceData = Binding.GetInstanceData<FInstanceData>();
+		check(InstanceData);
+
+		InstanceData->DeltaTime = TraitState.GetDeltaTime();
+
+		IUpdate::PreUpdate(Context, Binding, TraitState);
+	}
+
+	void FStrafeWarpingTrait::PostEvaluate(FEvaluateTraversalContext& Context, const TTraitBinding<IEvaluate>& Binding) const
+	{
+		IEvaluate::PostEvaluate(Context, Binding);
+
+		const FSharedData* SharedData = Binding.GetSharedData<FSharedData>();
+		check(SharedData);
+
+		FInstanceData* InstanceData = Binding.GetInstanceData<FInstanceData>();
+		check(InstanceData);
+
+		InstanceData->RootBoneTransform = SharedData->GetRootBoneTransform(Binding);
+
+		// Update target orientation, root bone transform, & other properties 
+		InstanceData->TargetOrientation = SharedData->GetTargetOrientation(Binding);
+		InstanceData->Alpha = SharedData->GetAlpha(Binding);
+
+#if ENABLE_ANIM_DEBUG 
+		InstanceData->HostObject = Context.GetHostObject();
+#endif // ENABLE_ANIM_DEBUG 
+
+		Context.AppendTask(FAnimNextStrafeWarpingTask::Make(InstanceData, SharedData));
+	}
+
+	void FStrafeWarpingTrait::InitializeSpineData(TArrayView<FSpineBoneData> OutSpineBoneData, const TArray<FName>& SpineBoneNames, const FLODPoseStack& Pose)
+	{
+		QUICK_SCOPE_CYCLE_COUNTER(FStrafeWarpingTrait_InitializeSpineData);
+		check(OutSpineBoneData.Num() == SpineBoneNames.Num())
+		
+		const FReferencePose& RefPose = *Pose.RefPose;
+
+		if (SpineBoneNames.Num() == 0)
+		{
+			return;
+		}
+		
+		for (int32 i = 0; i < SpineBoneNames.Num(); i++)
+		{
+			OutSpineBoneData[i].Weight = 0.0f;
+			if (const FBoneIndexType* SpineBoneIndex = RefPose.GetBoneNameToLODBoneIndexMap().Find(SpineBoneNames[i]))
+			{
+				OutSpineBoneData[i].LODBoneIndex = *SpineBoneIndex;
+			}
+			else
+			{
+				OutSpineBoneData[i].LODBoneIndex = INDEX_NONE;
+			}
+		}
+
+		// Calculate weight
+
+		// Sort bones indices so we can transform parent before child
+		OutSpineBoneData.Sort(FSpineBoneData::FCompareBoneIndex());
+
+		// Assign Weights.
+		TArray<int32, TInlineAllocator<20>> IndicesToUpdate;
+
+		// Note reverse iteration
+		for (int32 Index = OutSpineBoneData.Num() - 1; Index >= 0; Index--)
+		{
+			// If this bone's weight hasn't been updated, scan its parents.
+			// If parents have weight, we add it to 'ExistingWeight'.
+			// split (1.f - 'ExistingWeight') between all members of the chain that have no weight yet.
+			if (OutSpineBoneData[Index].Weight == 0.f)
+			{
+				IndicesToUpdate.Reset(OutSpineBoneData.Num());
+				float ExistingWeight = 0.f;
+				IndicesToUpdate.Add(Index);
+
+				for (int32 ParentIndex = Index - 1; ParentIndex >= 0; ParentIndex--)
+				{
+					if (Pose.IsBoneChildOf(OutSpineBoneData[Index].LODBoneIndex, OutSpineBoneData[ParentIndex].LODBoneIndex))
+					{
+						if (OutSpineBoneData[ParentIndex].Weight > 0.f)
+						{
+							ExistingWeight += OutSpineBoneData[ParentIndex].Weight;
+						}
+						else
+						{
+							IndicesToUpdate.Add(ParentIndex);
+						}
+					}
+				}
+
+				check(IndicesToUpdate.Num() > 0);
+				const float WeightToShare = 1.f - ExistingWeight;
+				const float IndividualWeight = WeightToShare / float(IndicesToUpdate.Num());
+
+				for (int32 UpdateListIndex = 0; UpdateListIndex < IndicesToUpdate.Num(); UpdateListIndex++)
+				{
+					OutSpineBoneData[IndicesToUpdate[UpdateListIndex]].Weight = IndividualWeight;
+				}
+			}
+		}
+	}
+} // namespace UE::AnimNext
+
+
+////////////////////////////////////////////////////////////////////////////////////////////////////////////////////////
+// FAnimNextSteeringTask
+
+FAnimNextStrafeWarpingTask FAnimNextStrafeWarpingTask::Make(UE::AnimNext::FStrafeWarpingTrait::FInstanceData* InstanceData, const UE::AnimNext::FStrafeWarpingTrait::FSharedData* SharedData)
+{
+	FAnimNextStrafeWarpingTask Task;
+	Task.InstanceData = InstanceData;
+	Task.SharedData = SharedData;
+	return Task;
+}
+
+
+
+void FAnimNextStrafeWarpingTask::Execute(UE::AnimNext::FEvaluationVM& VM) const
+{
+	using namespace UE::AnimNext;
+	QUICK_SCOPE_CYCLE_COUNTER(FAnimNextStrafeWarpingTask_Execute);
+	
+	if (InstanceData->DeltaTime == 0.f)
+	{
+		return;
+	}
+
+	if (InstanceData->Alpha == 0.f || !bAnimNextStrafeWarpingTraitEnabled)
+	{
+		return;
+	}
+
+	float TargetOrientationAngleRad = 0.0f;
+	const FVector RotationAxisVector = GetAxisVector(SharedData->RotationAxis);
+
+	// Target Orientation is in world space, transform to root
+	FQuat TargetRotation = InstanceData->RootBoneTransform.GetRotation().Inverse() * InstanceData->TargetOrientation; 
+	FVector TargetMoveDir = TargetRotation.GetForwardVector();
+
+	// Flatten locomotion direction, along the rotation axis.
+	TargetMoveDir = (TargetMoveDir - RotationAxisVector.Dot(TargetMoveDir) * RotationAxisVector).GetSafeNormal();
+
+	TUniquePtr<FKeyframeState>* Keyframe = VM.PeekValueMutable<TUniquePtr<FKeyframeState>>(KEYFRAME_STACK_NAME, 0);
+	if (Keyframe == nullptr)
+	{
+		return;
+	}
+
+	DoWarpRootMotion(*Keyframe, RotationAxisVector, TargetMoveDir, OUT TargetOrientationAngleRad);
+
+	// Calculate the orientation warp angle for pose adjustments (spine and foot IK)
+	const float MaxAngleCorrectionRad = FMath::DegreesToRadians(SharedData->MaxCorrectionDegrees);
+	
+	// Optionally interpolate the effective orientation towards the target orientation angle
+	// When the orientation warping node becomes relevant, the input pose orientation may not be aligned with the desired orientation.
+	// Instead of interpolating this difference, snap to the desired orientation if it's our first update to minimize corrections over-time.
+	if ((SharedData->RotationInterpSpeed > 0.f) /*&& !bIsFirstUpdate */)
+	{
+		const float SmoothOrientationAngleRad = FMath::FInterpTo(InstanceData->OrientationAngleForPoseWarpRad, TargetOrientationAngleRad, InstanceData->DeltaTime, SharedData->RotationInterpSpeed);
+		// Limit our interpolation rate to prevent pops.
+		// @TODO: Use better, more physically accurate interpolation here.
+		InstanceData->OrientationAngleForPoseWarpRad = FMath::Clamp(SmoothOrientationAngleRad, InstanceData->OrientationAngleForPoseWarpRad - MaxAngleCorrectionRad, InstanceData->OrientationAngleForPoseWarpRad + MaxAngleCorrectionRad);
+	}
+	else
+	{
+		InstanceData->OrientationAngleForPoseWarpRad = TargetOrientationAngleRad;
+	}
+
+	InstanceData->OrientationAngleForPoseWarpRad = FMath::Clamp(InstanceData->OrientationAngleForPoseWarpRad, -MaxAngleCorrectionRad, MaxAngleCorrectionRad);
+	// Allow the alpha value of the node to affect the final rotation
+	InstanceData->OrientationAngleForPoseWarpRad *= InstanceData->Alpha;
+
+	if (FMath::IsNearlyZero(InstanceData->OrientationAngleForPoseWarpRad, KINDA_SMALL_NUMBER))
+	{
+		// No strafe angle, early out before hitting the pose modification code
+		return;
+	}
+
+	DoWarpPose(*Keyframe, RotationAxisVector);
+}
+
+void FAnimNextStrafeWarpingTask::DoWarpRootMotion(TUniquePtr<UE::AnimNext::FKeyframeState>& Keyframe, const FVector& RotationAxisVector, const FVector& TargetMoveDir, float& OutTargetOrientationAngleRad) const
+{
+	FTransform ThisFrameRootMotionTransform = FTransform::Identity;
+	const UE::Anim::IAnimRootMotionProvider* RootMotionProvider = UE::Anim::IAnimRootMotionProvider::Get();
+	if (!RootMotionProvider)
+	{
+		UE_LOG(LogAnimNextWarping, Error, TEXT("FAnimNextStrafeWarpingTask::Execute, missing RootMotionProvider"));
+		return;
+	}
+
+	if (RootMotionProvider->ExtractRootMotion(Keyframe.Get()->Attributes, ThisFrameRootMotionTransform))
+	{
+		FVector RootMotionDeltaTranslation = ThisFrameRootMotionTransform.GetTranslation();
+		const FQuat PreviousRootMotionDeltaRotation = InstanceData->RootMotionDeltaRotation;
+		InstanceData->RootMotionDeltaRotation = ThisFrameRootMotionTransform.GetRotation();
+		
+		const float RootMotionDeltaSpeed = RootMotionDeltaTranslation.Size() / InstanceData->DeltaTime;
+		if (RootMotionDeltaSpeed < SharedData->MinRootMotionSpeedThreshold)
+		{
+			// If we're under the threshold, snap orientation angle to 0, and let interpolation handle the delta
+			OutTargetOrientationAngleRad = 0.0f;
+		}
+		else
+		{
+			const FVector PreviousRootMotionDeltaDirection = InstanceData->RootMotionDeltaDirection;
+			// Hold previous direction if we can't calculate it from current move delta, because the root is no longer moving
+			InstanceData->RootMotionDeltaDirection = RootMotionDeltaTranslation.GetSafeNormal(UE_SMALL_NUMBER, PreviousRootMotionDeltaDirection);
+			OutTargetOrientationAngleRad = UE::AnimNext::SignedAngleRadBetweenNormals(
+				InstanceData->RootMotionDeltaDirection, TargetMoveDir, RotationAxisVector);
+
+			// Motion Matching may return an animation that deviates a lot from the movement direction (e.g movement direction going bwd and motion matching could return the fwd animation for a few frames)
+			// When that happens, since we use the delta between root motion and movement direction, we would be over-rotating the lower body and breaking the pose during those frames
+			// So, when that happens we use the inverse of the root motion direction to calculate our target rotation. 
+			// This feels a bit 'hacky' but its the only option I've found so far to mitigate the problem
+			if (SharedData->LocomotionAngleDeltaThreshold > 0.f)
+			{
+				if (FMath::Abs(FMath::RadiansToDegrees(OutTargetOrientationAngleRad)) > SharedData->LocomotionAngleDeltaThreshold)
+				{
+					OutTargetOrientationAngleRad = FMath::UnwindRadians(OutTargetOrientationAngleRad + FMath::DegreesToRadians(180.0f));
+					InstanceData->RootMotionDeltaDirection = -InstanceData->RootMotionDeltaDirection;
+				}
+			}
+			/* No prediction in first iteration for AnimNext. This code is copy-pasta from AnimNode_OrientationWarping
+			// If there is translation in predicted root motion, use it if the orientation error is less than current error
+			if (TargetTime > 0 && CurrentAnimAsset && !PredictedRootMotionDeltaTranslation.IsNearlyZero(UE_SMALL_NUMBER))
+			{
+				PredictedRootMotionDeltaTranslation.Normalize();
+				float PredictedOrientationErrorAngleRad = UE::Anim::SignedAngleRadBetweenNormals(
+					PredictedRootMotionDeltaTranslation, LocomotionForward, RotationAxisVector);
+
+				// The future orientation will often match the current, so add a small delta to avoid further testing for same values
+				if (FMath::Abs(PredictedOrientationErrorAngleRad) + UE_KINDA_SMALL_NUMBER < FMath::Abs(TargetOrientationAngleRad))
+				{
+					// Note: Don't update root motion direction, as the root motion direction is what we are playing, 
+					// not the future. Updating root motion direction will break counter compensate. 
+					// Also, we rely on the built in interp for continuity / smoothness
+#if ENABLE_ANIM_DEBUG || ENABLE_VISUAL_LOG
+					FutureRootMotionDeltaDirection = PredictedRootMotionDeltaTranslation;
+					bUsedFutureRootMotion = true;
+#endif
+					TargetOrientationAngleRad = PredictedOrientationErrorAngleRad;
+				}
+			}
+			*/
+
+			// Don't compensate interpolation by the root motion angle delta if the previous direction isn't valid.
+			if (SharedData->bCounterCompenstateInterpolationByRootMotion && !PreviousRootMotionDeltaDirection.IsNearlyZero(UE_SMALL_NUMBER))
+			{
+				float RootMotionDeltaAngleRad = 0.0f;
+				// Counter the interpolated orientation angle by the root motion direction angle delta.
+				// This prevents our interpolation from fighting the natural root motion that's flowing through the graph.
+				// To correctly measure the amount to counter, we need to unrotate our previous delta direction by our previous rotation
+				// As the previous direction delta is relative to the previous rotation delta
+				RootMotionDeltaAngleRad = UE::AnimNext::SignedAngleRadBetweenNormals(
+					InstanceData->RootMotionDeltaDirection, PreviousRootMotionDeltaRotation.UnrotateVector(PreviousRootMotionDeltaDirection),
+					RotationAxisVector);
+
+				// Root motion may have large deltas i.e. bad blends or sudden direction changes like pivots.
+				// If there's an instantaneous pop in root motion direction, this is likely a pivot.
+				const float MaxRootMotionDeltaToCompensateRad = FMath::DegreesToRadians(SharedData->MaxRootMotionDeltaToCompensateDegrees);
+				if (FMath::Abs(RootMotionDeltaAngleRad) < MaxRootMotionDeltaToCompensateRad)
+				{
+					InstanceData->CounterCompensateTargetAngleRad += RootMotionDeltaAngleRad;
+					float CounterCompensateAngle = FMath::FInterpTo(0, InstanceData->CounterCompensateTargetAngleRad, InstanceData->DeltaTime,
+					                                                SharedData->CounterCompensateInterpSpeed);
+					InstanceData->OrientationAngleForPoseWarpRad = FMath::UnwindRadians(InstanceData->OrientationAngleForPoseWarpRad + CounterCompensateAngle);
+					InstanceData->CounterCompensateTargetAngleRad -= CounterCompensateAngle;
+				}
+			}
+
+			// Rotate the root motion delta fully by the warped angle
+			const FVector WarpedRootMotionTranslationDelta = FQuat(RotationAxisVector, OutTargetOrientationAngleRad).RotateVector(
+				RootMotionDeltaTranslation);
+			ThisFrameRootMotionTransform.SetTranslation(WarpedRootMotionTranslationDelta);
+			
+			RootMotionProvider->OverrideRootMotion(ThisFrameRootMotionTransform, Keyframe.Get()->Attributes);
+		}
+	}
+
+#if ENABLE_VISUAL_LOG
+	if (FVisualLogger::IsRecording())
+	{
+		constexpr float DebugDrawScale = 1.f;
+		const FTransform ComponentTransform = InstanceData->RootBoneTransform;
+		
+		FVector DebugArrowOffset = FVector::ZAxisVector * DebugDrawScale;
+		//uint8 DebugAlpha = AnimNodeOrientationWarpingDebugTransparency ? 255 * BlendWeight : 255;
+		uint8 DebugAlpha = 255;
+		FColor DebugColor = FColor::Green;
+
+		// Draw debug shapes
+		{
+			const FVector ForwardDirection = ComponentTransform.GetRotation().RotateVector(TargetMoveDir);
+
+			UE_VLOG_CIRCLE_THICK(InstanceData->HostObject, "OrientationWarping", Display,
+				ComponentTransform.GetLocation() + DebugArrowOffset + ForwardDirection * 100.f * DebugDrawScale,
+				FVector::UpVector, 4.f * DebugDrawScale, DebugColor.WithAlpha(DebugAlpha), 1.0f, TEXT(""));
+			UE_VLOG_ARROW(InstanceData->HostObject, "OrientationWarping", Display,
+				ComponentTransform.GetLocation() + DebugArrowOffset,
+				ComponentTransform.GetLocation() + DebugArrowOffset + ForwardDirection * 100.f * DebugDrawScale,
+				FColor::Red.WithAlpha(DebugAlpha), TEXT(""));
+
+			const FVector RotationDirection = ComponentTransform.GetRotation().RotateVector(InstanceData->RootMotionDeltaDirection);
+
+			DebugArrowOffset += FVector::ZAxisVector * DebugDrawScale;
+
+			UE_VLOG_CIRCLE_THICK(InstanceData->HostObject, "OrientationWarping", Display,
+				ComponentTransform.GetLocation() + DebugArrowOffset + RotationDirection * 100.f * DebugDrawScale,
+				FVector::UpVector, 4.f * DebugDrawScale, DebugColor.WithAlpha(DebugAlpha), 1.0f, TEXT(""));
+			UE_VLOG_ARROW(InstanceData->HostObject, "OrientationWarping", Display,
+				ComponentTransform.GetLocation() + DebugArrowOffset,
+				ComponentTransform.GetLocation() + DebugArrowOffset + RotationDirection * 100.f * DebugDrawScale,
+				FColor::Blue.WithAlpha(DebugAlpha), TEXT(""));
+
+			/* Debug for predictive root motion, not currently implemented
+			if (bUsedFutureRootMotion && bGraphDrivenWarping)
+			{
+				const FVector FutureRotationDirection = ComponentTransform.GetRotation().RotateVector(FutureRootMotionDeltaDirection);
+
+				UE_VLOG_CIRCLE_THICK(Output.AnimInstanceProxy->GetAnimInstanceObject(), "OrientationWarping", Display,
+					ComponentTransform.GetLocation() + DebugArrowOffset + FutureRotationDirection * 100.f * DebugDrawScale,
+					FVector::UpVector, 4.f * DebugDrawScale, DebugColor.WithAlpha(DebugAlpha), 1.0f, TEXT(""));
+				UE_VLOG_ARROW(Output.AnimInstanceProxy->GetAnimInstanceObject(), "OrientationWarping", Display,
+					ComponentTransform.GetLocation() + DebugArrowOffset,
+					ComponentTransform.GetLocation() + DebugArrowOffset + FutureRotationDirection * 100.f * DebugDrawScale,
+					FColor::Yellow.WithAlpha(DebugAlpha), TEXT(""));
+			}
+			*/
+
+			const float ActualOrientationAngleDegrees = FMath::RadiansToDegrees(InstanceData->OrientationAngleForPoseWarpRad);
+			const FVector WarpedRotationDirection = RotationDirection.RotateAngleAxis(ActualOrientationAngleDegrees, RotationAxisVector);
+
+			DebugArrowOffset += FVector::ZAxisVector * DebugDrawScale;
+
+			UE_VLOG_CIRCLE_THICK(InstanceData->HostObject, "OrientationWarping", Display,
+				ComponentTransform.GetLocation() + DebugArrowOffset + WarpedRotationDirection * 100.f * DebugDrawScale,
+				FVector::UpVector, 4.f * DebugDrawScale, DebugColor.WithAlpha(DebugAlpha), 1.0f, TEXT(""));
+			UE_VLOG_ARROW(InstanceData->HostObject, "OrientationWarping", Display,
+				ComponentTransform.GetLocation() + DebugArrowOffset,
+				ComponentTransform.GetLocation() + DebugArrowOffset + WarpedRotationDirection * 100.f * DebugDrawScale,
+				FColor::Green.WithAlpha(DebugAlpha), TEXT(""));
+		}
+	}
+#endif // ENABLE_VISUAL_LOG
+}
+
+void FAnimNextStrafeWarpingTask::DoWarpPose(TUniquePtr<UE::AnimNext::FKeyframeState>& Keyframe, const FVector& RotationAxisVector) const
+{
+	using namespace UE::AnimNext;
+	
+	const FLODPoseStack& Pose = Keyframe->Pose;
+
+	FTransformArrayView PoseTransforms = Pose.LocalTransformsView;
+	if (PoseTransforms.IsEmpty())
+	{
+		// No bones
+		return;
+	}
+
+	if (Pose.RefPose == nullptr)
+	{
+		// No ref pose, cannot continue
+		return;
+	}
+
+	const TArrayView<const FBoneIndexType>& LODBoneIndexToParentLODBoneIndexMap = Pose.GetLODBoneIndexToParentLODBoneIndexMap();
+	
+	FTransformArraySoAStack ComponentSpaceTransforms(PoseTransforms.Num(), false);
+	ComponentSpaceTransforms.CopyTransforms(PoseTransforms, 0, PoseTransforms.Num());
+	ConvertPoseLocalToMeshRotationTranslation(ComponentSpaceTransforms.GetView(), LODBoneIndexToParentLODBoneIndexMap);
+
+	const float RootOffset = FMath::UnwindRadians(InstanceData->OrientationAngleForPoseWarpRad * SharedData->DistributedBoneOrientationAlpha);
+
+	const float IKFootRootOrientationAlpha = 1.f - SharedData->DistributedBoneOrientationAlpha;
+
+	// Rotate IK Foot Root
+	if (!FMath::IsNearlyZero(IKFootRootOrientationAlpha, KINDA_SMALL_NUMBER))
+	{
+		{
+			const FQuat IKRootBoneRotation = FQuat(RotationAxisVector, InstanceData->OrientationAngleForPoseWarpRad * IKFootRootOrientationAlpha);
+
+			// IK Feet 
+			// We want these to keep their original component space orientation
+			// But we want them to translate based on some rotation
+			const int32 NumIKFootBones = SharedData->FootData.Num();
+
+			if (NumIKFootBones > 0)
+			{
+				for (int32 ArrayIndex = 0; ArrayIndex < NumIKFootBones; ArrayIndex++)
+				{
+					const FStrafeWarpFootData& FootData = SharedData->FootData[ArrayIndex];
+					
+					FBoneIndexType LegRootIndex = Pose.FindLODBoneIndexFromBoneName(FootData.LegRoot);
+					FBoneIndexType LegMidIndex = Pose.FindLODBoneIndexFromBoneName(FootData.LegMid);
+					FBoneIndexType LegTipIndex = Pose.FindLODBoneIndexFromBoneName(FootData.LegTip);
+
+					// Validate data
+					if (LegRootIndex == INDEX_NONE || LegMidIndex == INDEX_NONE || LegTipIndex == INDEX_NONE)
+					{
+						continue;
+					}
+
+					// Todo: should these be ensure? error? how to handle?
+					check(Pose.GetLODBoneParentIndex(LegTipIndex) == LegMidIndex);
+					check(Pose.GetLODBoneParentIndex(LegMidIndex) == LegRootIndex);
+
+					FTransform LegRootTransformCS(ComponentSpaceTransforms.Rotations[LegRootIndex], ComponentSpaceTransforms.Translations[LegRootIndex]);
+					FTransform LegMidTransformCS(ComponentSpaceTransforms.Rotations[LegMidIndex], ComponentSpaceTransforms.Translations[LegMidIndex]);
+					FTransform LegTipTransformCS(ComponentSpaceTransforms.Rotations[LegTipIndex], ComponentSpaceTransforms.Translations[LegTipIndex]);
+
+					FVector FootTargetPosition = IKRootBoneRotation.RotateVector(LegTipTransformCS.GetLocation());
+
+					// Joint (knee) target vector, needs to be driven from animation and define a good solving plan
+					FVector JointTarget = FootTargetPosition; // TODO This might break if leg is straight 
+
+					AnimationCore::SolveTwoBoneIK(LegRootTransformCS,
+						LegMidTransformCS,
+						LegTipTransformCS,
+						JointTarget,
+						FootTargetPosition,
+						false,
+						0.0,
+						0.0
+						);
+
+					// Apply results
+
+					// Be careful reverse TTransform vs FQuat order of operations
+					FTransform LocalTip = LegTipTransformCS * LegMidTransformCS.Inverse();
+					FTransform LocalMid = LegMidTransformCS * LegRootTransformCS.Inverse();
+
+					// Figure out local root by applying difference in CS
+					FQuat LocalRootRotationDiff = ComponentSpaceTransforms[LegRootIndex].Rotation.Inverse() * LegRootTransformCS.GetRotation(); 
+
+					PoseTransforms[LegTipIndex].Translation = LocalTip.GetLocation();
+					PoseTransforms[LegTipIndex].Rotation = LocalTip.GetRotation();
+					PoseTransforms[LegMidIndex].Translation = LocalMid.GetLocation();
+					PoseTransforms[LegMidIndex].Rotation = LocalMid.GetRotation();
+					PoseTransforms[LegRootIndex].Rotation = PoseTransforms[LegRootIndex].Rotation * LocalRootRotationDiff;
+					
+
+#if ENABLE_VISUAL_LOG
+					if (FVisualLogger::IsRecording())
+					{
+						const float DebugDrawScale = 10.0f;
+						const uint16 DebugDrawThickness = 1;
+						const float DebugDrawSpehereRadius = 10.f;
+						// Careful: FTransform multiplication is reveresed from quat!
+						FTransform FootWorldTransformOriginal = FTransform(ComponentSpaceTransforms.Rotations[LegTipIndex], ComponentSpaceTransforms.Translations[LegTipIndex]) * InstanceData->RootBoneTransform;
+						FTransform FootWorldTransformTarget = FTransform(ComponentSpaceTransforms.Rotations[LegTipIndex], FootTargetPosition) * InstanceData->RootBoneTransform;
+						FTransform FootWorldTransformSolved = LegTipTransformCS * InstanceData->RootBoneTransform;
+						UE_VLOG_COORDINATESYSTEM(InstanceData->HostObject, "OrientationWarping", Display, FootWorldTransformTarget.GetLocation(), FootWorldTransformTarget.GetRotation().Rotator(), DebugDrawScale, FColor::Green, DebugDrawThickness, TEXT(""));
+						UE_VLOG_COORDINATESYSTEM(InstanceData->HostObject, "OrientationWarping", Display, FootWorldTransformOriginal.GetLocation(), FootWorldTransformOriginal.GetRotation().Rotator(), DebugDrawScale, FColor::Green, DebugDrawThickness, TEXT(""));
+						UE_VLOG_ARROW(InstanceData->HostObject, "OrientationWarping", Display, FootWorldTransformOriginal.GetLocation(), FootWorldTransformTarget.GetLocation(), FColor::Magenta, TEXT(""));
+						UE_VLOG_SPHERE(InstanceData->HostObject, "OrientationWarping", Display, FootWorldTransformSolved.GetTranslation(), DebugDrawSpehereRadius, FColor::Yellow, TEXT(""));
+					}
+#endif
+				}
+			}
+		}
+	}
+	
+	// Rotate Root Bone first, as that cheaply rotates the whole pose with one transformation.
+	// We do this with the pose in local space since we want it to propagate
+	if (!FMath::IsNearlyZero(RootOffset, KINDA_SMALL_NUMBER))
+	{
+		if (SharedData->PreserveOriginalRootRotation)
+		{
+			// Find all children of the root and adjust them
+			for (FBoneIndexType BoneIndex = 1; BoneIndex < Pose.GetNumBones(); BoneIndex++)
+			{
+				if (Pose.GetLODBoneParentIndex(BoneIndex) == 0)
+				{
+					// Is a child of the root
+					const FVector LocalRotationVector = ComponentSpaceTransforms[BoneIndex].Rotation.UnrotateVector(RotationAxisVector);
+					const FQuat RootRotation = FQuat(LocalRotationVector, RootOffset);
+					PoseTransforms.Rotations[BoneIndex] = PoseTransforms.Rotations[BoneIndex] * RootRotation;
+				}
+			}
+		}
+		else
+		{
+			const FQuat RootRotation = FQuat(RotationAxisVector, RootOffset);
+
+			PoseTransforms.Rotations[0] = PoseTransforms.Rotations[0] * RootRotation;
+			PoseTransforms.Rotations[0].Normalize();
+		}
+	}
+
+	const int32 NumSpineBones = SharedData->SpineBones.Num();
+	const bool bSpineOrientationAlpha = !FMath::IsNearlyZero(SharedData->DistributedBoneOrientationAlpha, KINDA_SMALL_NUMBER);
+	const bool bUpdateSpineBones = (NumSpineBones > 0) && bSpineOrientationAlpha;
+
+	if (bUpdateSpineBones)
+	{
+		// Todo: Can we get away with lazy init here? Does the ref pose skeleton change at runtime?
+		// Todo: Cache spine bone data
+		TArray<FStrafeWarpingTrait::FSpineBoneData, TInlineAllocator<16>> SpineBoneDataArray;
+		SpineBoneDataArray.SetNumUninitialized(NumSpineBones);
+		FStrafeWarpingTrait::InitializeSpineData(TArrayView<FStrafeWarpingTrait::FSpineBoneData>(SpineBoneDataArray), SharedData->SpineBones, Pose);
+		
+		// Spine bones counter rotate body orientation evenly across all bones.
+		// Note: reverse iteration is important! We go from child to parent
+		for (int32 ArrayIndex = NumSpineBones - 1; ArrayIndex >= 0; ArrayIndex--)
+		{
+			const FStrafeWarpingTrait::FSpineBoneData& BoneData = SpineBoneDataArray[ArrayIndex];
+			if (BoneData.LODBoneIndex == INDEX_NONE || BoneData.Weight == 0.0f)
+			{
+				continue;
+			}
+			
+			// Important note! The root was moved in local space, so our component transform array is actually out of date
+			// However since we know everything rotated around RotationAxisVector, it doesn't matter for this calculation
+			const FVector LocalRotationVector = ComponentSpaceTransforms[BoneData.LODBoneIndex].Rotation.UnrotateVector(RotationAxisVector);
+			const FQuat SpineBoneCounterRotation = FQuat(LocalRotationVector, -InstanceData->OrientationAngleForPoseWarpRad * SharedData->DistributedBoneOrientationAlpha * BoneData.Weight);
+
+			PoseTransforms.Rotations[BoneData.LODBoneIndex] = PoseTransforms.Rotations[BoneData.LODBoneIndex] * SpineBoneCounterRotation;
+			PoseTransforms.Rotations[BoneData.LODBoneIndex].Normalize();
+		}
+	}
+}
+
diff --git a/Engine/Plugins/Experimental/AnimNextWarping/Source/AnimNextWarping/Private/StrafeWarpingTrait.h b/Engine/Plugins/Experimental/AnimNextWarping/Source/AnimNextWarping/Private/StrafeWarpingTrait.h
new file mode 100644
--- /dev/null
+++ b/Engine/Plugins/Experimental/AnimNextWarping/Source/AnimNextWarping/Private/StrafeWarpingTrait.h
@@ -0,0 +1,228 @@
+// Copyright Epic Games, Inc. All Rights Reserved.
+
+#pragma once
+
+#include "CoreMinimal.h"
+#include "EngineDefines.h"
+#include "TraitCore/Trait.h"
+#include "TraitCore/TraitSharedData.h"
+#include "TraitInterfaces/IAttributeProvider.h"
+#include "TraitInterfaces/IEvaluate.h"
+#include "TraitInterfaces/IUpdate.h"
+#include "StrafeWarpingTrait.generated.h"
+
+// Small structure to define a simple IK setup for legs
+USTRUCT()
+struct FStrafeWarpFootData
+{
+	GENERATED_BODY()
+
+	// e.g. the thigh bone
+	UPROPERTY(EditAnywhere, Category=Settings)
+	FName LegRoot;
+
+	// e.g. the knee bone
+	UPROPERTY(EditAnywhere, Category=Settings)
+	FName LegMid;
+
+	// e.g. the foot bone
+	UPROPERTY(EditAnywhere, Category=Settings)
+	FName LegTip;
+};
+
+USTRUCT(meta = (DisplayName = "Strafe Warping"))
+struct FStrafeWarpingTraitSharedData : public FAnimNextTraitSharedData
+{
+	GENERATED_BODY()
+
+	// Current strength of the skeletal control
+	UPROPERTY(EditAnywhere, Category = Alpha, meta = (PinShownByDefault))
+	float Alpha = 1.0f;
+
+	// The Orientation to steer towards
+	UPROPERTY(EditAnywhere, Category=Evaluation, meta=(PinShownByDefault))
+	FQuat TargetOrientation = FQuat::Identity;
+	
+	// @TODO Temp / try to remove this. Shouldn't have to feed as argument
+	// Last root bone transform sampled
+	UPROPERTY(EditAnywhere, Category = Evaluation, meta = (PinShownByDefault))
+	FTransform RootBoneTransform = FTransform::Identity;
+
+	// Rotation axis used when rotating the character body
+	UPROPERTY(EditAnywhere, Category=Settings, meta=(Inline))
+	TEnumAsByte<EAxis::Type> RotationAxis = EAxis::Z;
+
+	// Specifies how much rotation is applied to the character body versus IK feet
+	UPROPERTY(EditAnywhere, Category=Settings, meta=(Inline, ClampMin="0.0", ClampMax="1.0", PinHiddenByDefault))
+	float DistributedBoneOrientationAlpha = 0.5f;
+
+	// Specifies the interpolation speed (in Alpha per second) towards reaching the final warped rotation angle
+	// A value of 0 will cause instantaneous rotation, while a greater value will introduce smoothing
+	UPROPERTY(EditAnywhere, Category=Settings, meta=(Inline, ClampMin="0.0"))
+	float RotationInterpSpeed = 10.f;
+
+	// Same as RotationInterpSpeed, but for CounterCompensate smoothing. A value of 0 sample raw root motion.
+	// Used to avoid stuttering from resampling root deltas. Root motion is already smooth, so a large value is our default (~75% of 60 fps).
+	UPROPERTY(EditAnywhere, Category = Settings, meta = (Inline, ClampMin = "0.0"))
+	float CounterCompensateInterpSpeed = 45.f;
+
+	// Max correction we're allowed to do per-second when using interpolation.
+	// This minimizes pops when we have a large difference between current and target orientation.
+	UPROPERTY(EditAnywhere, Category=Settings, meta=(Inline, ClampMin="0.0", EditCondition="RotationInterpSpeed > 0.0f"))
+	float MaxCorrectionDegrees = 180.f;
+
+	// Don't compensate our interpolator when the instantaneous root motion delta is higher than this. This is likely a pivot.
+	UPROPERTY(EditAnywhere, Category=Settings, meta=(Inline, ClampMin="0.0", EditCondition="RotationInterpSpeed > 0.0f"))
+	float MaxRootMotionDeltaToCompensateDegrees = 45.f;
+
+	// Whether to counter compensate interpolation by the animated root motion angle change over time.
+	// This helps to conserve the motion from our animation.
+	// Disable this if your root motion is expected to be jittery, and you want orientation warping to smooth it out.
+	UPROPERTY(EditAnywhere, Category=Settings, meta=(Inline, EditCondition="RotationInterpSpeed > 0.0f"))
+	bool bCounterCompenstateInterpolationByRootMotion = true;
+	
+	// Minimum root motion speed required to apply orientation warping
+	// This is useful to prevent unnatural re-orientation when the animation has a portion with no root motion (i.e starts/stops/idles)
+	// When this value is greater than 0, it's recommended to enable interpolation with RotationInterpSpeed > 0
+	UPROPERTY(EditAnywhere, Category = Evaluation, meta = (ClampMin = "0.0", Inline))
+	float MinRootMotionSpeedThreshold = 10.0f;
+	
+	// Specifies an angle threshold to prevent erroneous over-rotation of the character, disabled with a value of 0
+	//
+	// When the effective orientation warping angle is detected to be greater than this value (default: 90 degrees) the locomotion direction will be inverted prior to warping
+	// This will be used in the following equation: [Orientation = RotationBetween(RootMotionDirection, -LocomotionDirection)]
+	//
+	// Example: Playing a forward running animation while the motion is going backward 
+	// Rather than orientation warping by 180 degrees, the system will warp by 0 degrees 
+	UPROPERTY(EditAnywhere, Category=Evaluation, meta=(Inline), meta=(ClampMin="0.0", ClampMax="180.0"))
+	float LocomotionAngleDeltaThreshold = 90.f;
+
+	// When true, propagates any modification on the root bone down to all of its children
+	// When false, will directly modify the root
+	UPROPERTY(EditAnywhere, Category=Evaluation, meta=(Inline))
+	bool PreserveOriginalRootRotation = true;
+
+	// Spine bone definitions
+	// Used to counter rotate the body in order to keep the character facing forward
+	// The amount of counter rotation applied is driven by DistributedBoneOrientationAlpha
+	// TODO: we ideally would use an equivlane to FBoneReference
+	UPROPERTY(EditAnywhere, Category=Settings, meta=(Inline))
+	TArray<FName> SpineBones;
+
+	UPROPERTY(EditAnywhere, Category=Settings, meta=(DisplayName="IK Foot Bones"), meta=(Inline, Hidden))
+	TArray<FStrafeWarpFootData> FootData;
+
+	// Latent pin support boilerplate
+	#define TRAIT_LATENT_PROPERTIES_ENUMERATOR(GeneratorMacro) \
+		GeneratorMacro(Alpha) \
+		GeneratorMacro(TargetOrientation) \
+		GeneratorMacro(RootBoneTransform) \
+
+	GENERATE_TRAIT_LATENT_PROPERTIES(FStrafeWarpingTraitSharedData, TRAIT_LATENT_PROPERTIES_ENUMERATOR)
+	#undef TRAIT_LATENT_PROPERTIES_ENUMERATOR
+};
+
+/**
+ * This trait modifies an animation pose to orientate towards a desired move direction 
+ */
+namespace UE::AnimNext
+{
+
+struct FStrafeWarpingTrait : FAdditiveTrait, IUpdate, IEvaluate
+{
+	DECLARE_ANIM_TRAIT(FStrafeWarpingTrait, FAdditiveTrait)
+
+	using FSharedData = FStrafeWarpingTraitSharedData;
+
+	// Internal structure with some precomputed weight values and bone indices
+	struct FSpineBoneData
+	{
+		FBoneIndexType LODBoneIndex;
+		float Weight;
+
+		FSpineBoneData()
+			: LODBoneIndex(INDEX_NONE)
+			, Weight(0.f)
+		{
+		}
+
+		FSpineBoneData(FBoneIndexType InBoneIndex)
+			: LODBoneIndex(InBoneIndex)
+			, Weight(0.f)
+		{
+		}
+
+		// Comparison Operator for Sorting
+		struct FCompareBoneIndex
+		{
+			FORCEINLINE bool operator()(const FSpineBoneData& A, const FSpineBoneData& B) const
+			{
+				return A.LODBoneIndex < B.LODBoneIndex;
+			}
+		};
+	};
+
+	struct FInstanceData : FTrait::FInstanceData
+	{
+		/** Target orientation per instance */
+		FQuat TargetOrientation = FQuat::Identity;
+		
+		/** Last root bone transform sampled */
+		FTransform RootBoneTransform = FTransform::Identity;
+
+		/** Delta in seconds between updates, populated during PreUpdate */
+		float DeltaTime = 0.f;
+
+		float Alpha = 1.0f;
+
+		// Internal current frame root motion delta direction
+		FVector RootMotionDeltaDirection = FVector::ZeroVector;
+
+		// Internal current frame root motion delta angle
+		FQuat RootMotionDeltaRotation = FQuat::Identity;
+
+		// Target for counter compenstate, we keep the target so we can smoothly interp.
+		float CounterCompensateTargetAngleRad = 0.0f;
+
+		// Internal orientation warping angle
+		float OrientationAngleForPoseWarpRad = 0.f;
+		
+#if ENABLE_ANIM_DEBUG 
+		/** Debug Object for VisualLogger */
+		TObjectPtr<const UObject> HostObject = nullptr;
+#endif // ENABLE_ANIM_DEBUG 
+	};
+
+	// IUpdate impl 
+	virtual void PreUpdate(FUpdateTraversalContext& Context, const TTraitBinding<IUpdate>& Binding, const FTraitUpdateState& TraitState) const override;
+
+	// IEvaluate impl
+	virtual void PostEvaluate(FEvaluateTraversalContext& Context, const TTraitBinding<IEvaluate>& Binding) const override;
+
+	// Helper for initialization of spine data
+	// Ideally this would be called once on trait become relevent, but we don't know the reference pose until Execute right now so must call it every frame
+	static void InitializeSpineData(TArrayView<FSpineBoneData> OutSpineBoneData, const TArray<FName>& SpineBoneNames, const FLODPoseStack& Pose);
+};
+
+} // namespace UE::AnimNext
+
+/** Task to run Strafe Warping on VM */
+USTRUCT()
+struct FAnimNextStrafeWarpingTask : public FAnimNextEvaluationTask
+{
+	GENERATED_BODY()
+
+	DECLARE_ANIM_EVALUATION_TASK(FAnimNextStrafeWarpingTask)
+
+	static FAnimNextStrafeWarpingTask Make(UE::AnimNext::FStrafeWarpingTrait::FInstanceData* InstanceData, const UE::AnimNext::FStrafeWarpingTrait::FSharedData* SharedData);
+
+	// Task entry point
+	virtual void Execute(UE::AnimNext::FEvaluationVM& VM) const override;
+
+private:
+	void DoWarpRootMotion(TUniquePtr<UE::AnimNext::FKeyframeState>& KeyframeState, const FVector& RotationAxisVector, const FVector& TargetMoveDir, float& OutTargetOrientationAngleRad) const;
+	void DoWarpPose(TUniquePtr<UE::AnimNext::FKeyframeState>& KeyframeState, const FVector& RotationAxisVector) const;
+
+	UE::AnimNext::FStrafeWarpingTrait::FInstanceData* InstanceData = nullptr;
+	const UE::AnimNext::FStrafeWarpingTrait::FSharedData* SharedData = nullptr;
+};
\ No newline at end of file
