Steering trait is supported on non-animation sequence nodes. This means it can be used in places where we dont care about the predictive root motion part of steering and only want to apply procedural

diff --git a/Engine/Plugins/Experimental/AnimNextWarping/Source/AnimNextWarping/Private/SteeringTrait.cpp b/Engine/Plugins/Experimental/AnimNextWarping/Source/AnimNextWarping/Private/SteeringTrait.cpp
--- a/Engine/Plugins/Experimental/AnimNextWarping/Source/AnimNextWarping/Private/SteeringTrait.cpp
+++ b/Engine/Plugins/Experimental/AnimNextWarping/Source/AnimNextWarping/Private/SteeringTrait.cpp
@@ -64,25 +64,24 @@ namespace UE::AnimNext
 		// Get Current Anim Asset
 		{
 			TTraitBinding<IAttributeProvider> AttributeTrait;
-			if (!Binding.GetStackInterface(AttributeTrait))
+			if (Binding.GetStackInterface(AttributeTrait))
 			{
-				UE_LOG(LogAnimNextWarping, Error, TEXT("FSteeringTrait::PostEvaluate, missing IAttributeProvider"));
-				return;
+				InstanceData->OnExtractRootMotionAttribute = AttributeTrait.GetOnExtractRootMotionAttribute(Context);
 			}
-
-			InstanceData->OnExtractRootMotionAttribute = AttributeTrait.GetOnExtractRootMotionAttribute(Context);
 		}
 
 		// Get Current Anim Asset Time
 		{
 			TTraitBinding<ITimeline> TimelineTrait;
-			if (!Binding.GetStackInterface(TimelineTrait))
+			if (Binding.GetStackInterface(TimelineTrait))
 			{
-				UE_LOG(LogAnimNextWarping, Error, TEXT("FSteeringTrait::PostEvaluate, missing ITimeline"));
-				return;
+				InstanceData->CurrentAnimAssetTime = TimelineTrait.GetState(Context).GetPosition();
+			}
+			else
+			{
+				// Set time to be -1.0 to indicate we should skip root motion prediction consideration
+				InstanceData->CurrentAnimAssetTime = -1.0f;
 			}
-
-			InstanceData->CurrentAnimAssetTime = TimelineTrait.GetState(Context).GetPosition();
 		}
 
 		// Update target orientation, root bone transform, & other properties 
@@ -152,8 +151,9 @@ void FAnimNextSteeringTask::Execute(UE::AnimNext::FEvaluationVM& VM) const
 
 						if (SharedData->AnimatedTargetTime > 0.0f)
 						{
-							if (InstanceData->OnExtractRootMotionAttribute.IsBound())
+							if (InstanceData->OnExtractRootMotionAttribute.IsBound() && InstanceData->CurrentAnimAssetTime != -1.0f)
 							{
+								check(InstanceData->CurrentAnimAssetTime >= 0.0f);
 								FTransform PredictedRootMotionTransform = InstanceData->OnExtractRootMotionAttribute.Execute(InstanceData->CurrentAnimAssetTime, SharedData->AnimatedTargetTime, true);
 								FQuat PredictedRootMotionQuat = PredictedRootMotionTransform.GetRotation();
 								FRotator PredictedRootMotionRot = FRotator(PredictedRootMotionQuat);
